#TTP 331191 - Texture Editor: Reimport button should be disabled when viewing Lightmap textures

#branch UE4

#proj Editor.TextureEditor

#summary Reimport button in the texture editor is now disabled for light and shadow maps

#add Added a can execute method to the reimport UICommand which checks the texture type and disabled the function if it's a light of shadow map.

diff --git a/Engine/Source/Editor/TextureEditor/Private/TextureEditorToolkit.cpp b/Engine/Source/Editor/TextureEditor/Private/TextureEditorToolkit.cpp
--- a/Engine/Source/Editor/TextureEditor/Private/TextureEditorToolkit.cpp
+++ b/Engine/Source/Editor/TextureEditor/Private/TextureEditorToolkit.cpp
@@ -666,7 +666,8 @@ void FTextureEditorToolkit::BindCommands()
 
 	ToolkitCommands->MapAction(
 		Commands.Reimport,
-		FExecuteAction::CreateSP(this, &FTextureEditorToolkit::OnReimport));
+		FExecuteAction::CreateSP(this, &FTextureEditorToolkit::OnReimport),
+		FCanExecuteAction::CreateSP(this, &FTextureEditorToolkit::OnReimportEnabled));
 
 	ToolkitCommands->MapAction(
 		Commands.Settings,
@@ -870,6 +871,14 @@ void FTextureEditorToolkit::OnReimport()
 	FReimportManager::Instance()->Reimport(Texture, /*bAskForNewFileIfMissing=*/true);
 }
 
+bool FTextureEditorToolkit::OnReimportEnabled() const
+{
+	if ( Texture->IsA<ULightMapTexture2D>() || Texture->IsA<UShadowMapTexture2D>() )
+	{
+		return false;
+	}
+	return true;
+}
 
 void FTextureEditorToolkit::HandleSettingsActionExecute()
 {
diff --git a/Engine/Source/Editor/TextureEditor/Private/TextureEditorToolkit.h b/Engine/Source/Editor/TextureEditor/Private/TextureEditorToolkit.h
--- a/Engine/Source/Editor/TextureEditor/Private/TextureEditorToolkit.h
+++ b/Engine/Source/Editor/TextureEditor/Private/TextureEditorToolkit.h
@@ -164,6 +164,7 @@ private:
 	void OnCompressNow();
 	bool OnCompressNowEnabled() const;
 	void OnReimport();
+	bool OnReimportEnabled() const;
 	void HandleSettingsActionExecute();
 	bool IsCubeTexture() const;
 
