Disabled digit grouping in the texture details text.
#ttp 332686 - Regression: EDITOR: TEXTURE: Resolution statistics are comma delimited
#branch UE4
#proj Editor.TextureEditor
#add Created an FNumberFormattingOptions with digit grouping disabled, and passed as parameter to FText::AsNumber in FTextureEditorToolkit::PopulateQuickInfo.

diff --git a/Engine/Source/Editor/TextureEditor/Private/TextureEditorToolkit.cpp b/Engine/Source/Editor/TextureEditor/Private/TextureEditorToolkit.cpp
--- a/Engine/Source/Editor/TextureEditor/Private/TextureEditorToolkit.cpp
+++ b/Engine/Source/Editor/TextureEditor/Private/TextureEditorToolkit.cpp
@@ -292,9 +292,12 @@ void FTextureEditorToolkit::PopulateQuickInfo( )
 
 	CalculateTextureDimensions(Width, Height);
 
-	ImportedText->SetText( FText::Format( NSLOCTEXT("TextureEditor", "QuickInfo_Imported", "Imported: {0}x{1}"), FText::AsNumber( ImportedWidth ), FText::AsNumber( ImportedHeight ) ) );
-	CurrentText->SetText(  FText::Format( NSLOCTEXT("TextureEditor", "QuickInfo_Current", "Current: {0}x{1}"), FText::AsNumber( FMath::Max((uint32)1, ImportedWidth >> MipLevel) ), FText::AsNumber( FMath::Max((uint32)1, ImportedHeight >> MipLevel)) ) );
-	MaxInGameText->SetText(FText::Format( NSLOCTEXT("TextureEditor", "QuickInfo_MaxInGame", "Max In-Game: {0}x{1}"), FText::AsNumber( MaxInGameWidth ), FText::AsNumber( MaxInGameHeight )));
+	FNumberFormattingOptions Options;
+	Options.UseGrouping = false;
+
+	ImportedText->SetText( FText::Format( NSLOCTEXT("TextureEditor", "QuickInfo_Imported", "Imported: {0}x{1}"), FText::AsNumber( ImportedWidth, &Options ), FText::AsNumber( ImportedHeight, &Options ) ) );
+	CurrentText->SetText(  FText::Format( NSLOCTEXT("TextureEditor", "QuickInfo_Current", "Current: {0}x{1}"), FText::AsNumber( FMath::Max((uint32)1, ImportedWidth >> MipLevel), &Options ), FText::AsNumber( FMath::Max((uint32)1, ImportedHeight >> MipLevel), &Options) ) );
+	MaxInGameText->SetText(FText::Format( NSLOCTEXT("TextureEditor", "QuickInfo_MaxInGame", "Max In-Game: {0}x{1}"), FText::AsNumber( MaxInGameWidth, &Options ), FText::AsNumber( MaxInGameHeight, &Options )));
 	MethodText->SetText(   FText::Format( NSLOCTEXT("TextureEditor", "QuickInfo_Method", "Method: {0}"), Texture->NeverStream ? NSLOCTEXT("TextureEditor", "QuickInfo_MethodNotStreamed", "Not Streamed") : NSLOCTEXT("TextureEditor", "QuickInfo_MethodStreamed", "Streamed")));
 	LODBiasText->SetText(  FText::Format( NSLOCTEXT("TextureEditor", "QuickInfo_LODBias", "Combined LOD Bias: {0}"), FText::AsNumber( Texture->CachedCombinedLODBias )));
 	
