Fixed crash caused by the texture editor viewport trying to access a texture during reimport

TTP# 332724 - Regression: EDITOR: Texture: CRASH: Reimporting a texture that is not a power of two results in a crash

Reimporting a non-power of 2 texture would cause a modal dialog to open, causing the texture editor viewport to be ticked for render while the texture was still being reimported. This lead to a crash due to the texture having no resource allocated.

The rendering in the texture editor is now disabled while a reimport takes place.

I also moved the Pre/PostReimport delegates from FAssetEditorToolkit to FReimportManager to ensure that all reimport cases are caught (eg, reimporting from the Content Browser while an asset editor is open).

diff --git a/Engine/Source/Editor/DestructibleMeshEditor/Private/DestructibleMeshEditor.cpp b/Engine/Source/Editor/DestructibleMeshEditor/Private/DestructibleMeshEditor.cpp
--- a/Engine/Source/Editor/DestructibleMeshEditor/Private/DestructibleMeshEditor.cpp
+++ b/Engine/Source/Editor/DestructibleMeshEditor/Private/DestructibleMeshEditor.cpp
@@ -35,7 +35,7 @@ void FDestructibleMeshEditorCommands::RegisterCommands()
 
 FDestructibleMeshEditor::~FDestructibleMeshEditor()
 {
-	FAssetEditorToolkit::OnPostReimport().RemoveAll(this);
+	FReimportManager::Instance()->OnPostReimport().RemoveAll(this);
 
 	GEditor->OnObjectReimported().RemoveAll(this);
 }
@@ -75,7 +75,7 @@ void FDestructibleMeshEditor::UnregisterTabSpawners(const TSharedRef<class FTabM
 
 void FDestructibleMeshEditor::InitDestructibleMeshEditor( const EToolkitMode::Type Mode, const TSharedPtr< class IToolkitHost >& InitToolkitHost, UDestructibleMesh* InDestructibleMesh )
 {
-	FAssetEditorToolkit::OnPostReimport().AddRaw(this, &FDestructibleMeshEditor::OnPostReimport);
+	FReimportManager::Instance()->OnPostReimport().AddRaw(this, &FDestructibleMeshEditor::OnPostReimport);
 
 	// Register our commands. This will only register them if not previously registered
 	FDestructibleMeshEditorCommands::Register();
diff --git a/Engine/Source/Editor/FontEditor/Private/FontEditor.cpp b/Engine/Source/Editor/FontEditor/Private/FontEditor.cpp
--- a/Engine/Source/Editor/FontEditor/Private/FontEditor.cpp
+++ b/Engine/Source/Editor/FontEditor/Private/FontEditor.cpp
@@ -104,7 +104,7 @@ void FFontEditor::UnregisterTabSpawners(const TSharedRef<class FTabManager>& Tab
 
 FFontEditor::~FFontEditor()
 {
-	FAssetEditorToolkit::OnPostReimport().RemoveAll(this);
+	FReimportManager::Instance()->OnPostReimport().RemoveAll(this);
 
 	UEditorEngine* Editor = (UEditorEngine*)GEngine;
 	if (Editor != NULL)
@@ -116,7 +116,7 @@ FFontEditor::~FFontEditor()
 
 void FFontEditor::InitFontEditor(const EToolkitMode::Type Mode, const TSharedPtr< class IToolkitHost >& InitToolkitHost, UObject* ObjectToEdit)
 {
-	FAssetEditorToolkit::OnPostReimport().AddRaw(this, &FFontEditor::OnPostReimport);
+	FReimportManager::Instance()->OnPostReimport().AddRaw(this, &FFontEditor::OnPostReimport);
 
 	// Register to be notified when an object is reimported.
 	GEditor->OnObjectReimported().AddSP(this, &FFontEditor::OnObjectReimported);
diff --git a/Engine/Source/Editor/Persona/Private/Persona.cpp b/Engine/Source/Editor/Persona/Private/Persona.cpp
--- a/Engine/Source/Editor/Persona/Private/Persona.cpp
+++ b/Engine/Source/Editor/Persona/Private/Persona.cpp
@@ -440,7 +440,7 @@ FPersona::FPersona()
 
 FPersona::~FPersona()
 {
-	FAssetEditorToolkit::OnPostReimport().RemoveAll(this);
+	FReimportManager::Instance()->OnPostReimport().RemoveAll(this);
 
 	FPersonaModule* PersonaModule = &FModuleManager::LoadModuleChecked<FPersonaModule>( "Persona" );
 	PersonaModule->GetMenuExtensibilityManager()->RemoveExtender(PersonaMenuExtender);
@@ -664,7 +664,7 @@ UObject* FPersona::GetAnimBlueprintAsObject() const
 
 void FPersona::InitPersona(const EToolkitMode::Type Mode, const TSharedPtr< class IToolkitHost >& InitToolkitHost, USkeleton* InitSkeleton, UAnimBlueprint* InitAnimBlueprint, UAnimationAsset* InitAnimationAsset, class USkeletalMesh * InitMesh)
 {
-	FAssetEditorToolkit::OnPostReimport().AddRaw(this, &FPersona::OnPostReimport);
+	FReimportManager::Instance()->OnPostReimport().AddRaw(this, &FPersona::OnPostReimport);
 
 	AssetDirtyBrush = FEditorStyle::GetBrush("ContentBrowser.ContentDirty");
 
diff --git a/Engine/Source/Editor/StaticMeshEditor/Private/StaticMeshEditor.cpp b/Engine/Source/Editor/StaticMeshEditor/Private/StaticMeshEditor.cpp
--- a/Engine/Source/Editor/StaticMeshEditor/Private/StaticMeshEditor.cpp
+++ b/Engine/Source/Editor/StaticMeshEditor/Private/StaticMeshEditor.cpp
@@ -73,7 +73,7 @@ void FStaticMeshEditor::UnregisterTabSpawners(const TSharedRef<class FTabManager
 
 FStaticMeshEditor::~FStaticMeshEditor()
 {
-	FAssetEditorToolkit::OnPostReimport().RemoveAll(this);
+	FReimportManager::Instance()->OnPostReimport().RemoveAll(this);
 
 	GEditor->UnregisterForUndo( this );
 	GEditor->OnObjectReimported().RemoveAll(this);
@@ -81,7 +81,7 @@ FStaticMeshEditor::~FStaticMeshEditor()
 
 void FStaticMeshEditor::InitStaticMeshEditor( const EToolkitMode::Type Mode, const TSharedPtr< class IToolkitHost >& InitToolkitHost, UStaticMesh* ObjectToEdit )
 {
-	FAssetEditorToolkit::OnPostReimport().AddRaw(this, &FStaticMeshEditor::OnPostReimport);
+	FReimportManager::Instance()->OnPostReimport().AddRaw(this, &FStaticMeshEditor::OnPostReimport);
 
 	// Support undo/redo
 	ObjectToEdit->SetFlags( RF_Transactional );
diff --git a/Engine/Source/Editor/TextureEditor/Private/TextureEditorToolkit.cpp b/Engine/Source/Editor/TextureEditor/Private/TextureEditorToolkit.cpp
--- a/Engine/Source/Editor/TextureEditor/Private/TextureEditorToolkit.cpp
+++ b/Engine/Source/Editor/TextureEditor/Private/TextureEditorToolkit.cpp
@@ -28,8 +28,8 @@ const FName FTextureEditorToolkit::PropertiesTabId(TEXT("TextureEditor_Propertie
 
 FTextureEditorToolkit::~FTextureEditorToolkit( )
 {
-	FAssetEditorToolkit::OnPreReimport().RemoveAll(this);
-	FAssetEditorToolkit::OnPostReimport().RemoveAll(this);
+	FReimportManager::Instance()->OnPreReimport().RemoveAll(this);
+	FReimportManager::Instance()->OnPostReimport().RemoveAll(this);
 
 	GEditor->UnregisterForUndo(this);
 }
@@ -65,8 +65,8 @@ void FTextureEditorToolkit::UnregisterTabSpawners( const TSharedRef<class FTabMa
 
 void FTextureEditorToolkit::InitTextureEditor( const EToolkitMode::Type Mode, const TSharedPtr< class IToolkitHost >& InitToolkitHost, UObject* ObjectToEdit )
 {
-	FAssetEditorToolkit::OnPreReimport().AddRaw(this, &FTextureEditorToolkit::OnPreReimport);
-	FAssetEditorToolkit::OnPostReimport().AddRaw(this, &FTextureEditorToolkit::OnPostReimport);
+	FReimportManager::Instance()->OnPreReimport().AddRaw(this, &FTextureEditorToolkit::OnPreReimport);
+	FReimportManager::Instance()->OnPostReimport().AddRaw(this, &FTextureEditorToolkit::OnPostReimport);
 
 	Texture = CastChecked<UTexture>(ObjectToEdit);
 
@@ -827,6 +827,9 @@ void FTextureEditorToolkit::OnPreReimport(UObject* InObject)
 	}
 
 	Texture->SourceFileTimestamp = TimeStamp.ToString();
+
+	// Disable viewport rendering until the texture has finished re-importing
+	TextureViewport->DisableRendering();
 }
 
 
@@ -843,6 +846,9 @@ void FTextureEditorToolkit::OnPostReimport(UObject* InObject, bool bSuccess)
 		// Failed, restore the compression flag
 		Texture->DeferCompression = SavedCompressionSetting;
 	}
+
+	// Re-enable viewport rendering now that the texture should be in a known state again
+	TextureViewport->EnableRendering();
 }
 
 
diff --git a/Engine/Source/Editor/TextureEditor/Private/Widgets/STextureEditorViewport.cpp b/Engine/Source/Editor/TextureEditor/Private/Widgets/STextureEditorViewport.cpp
--- a/Engine/Source/Editor/TextureEditor/Private/Widgets/STextureEditorViewport.cpp
+++ b/Engine/Source/Editor/TextureEditor/Private/Widgets/STextureEditorViewport.cpp
@@ -28,6 +28,7 @@ void STextureEditorViewport::AddReferencedObjects( FReferenceCollector& Collecto
 void STextureEditorViewport::Construct( const FArguments& InArgs, const TSharedRef<ITextureEditorToolkit>& InToolkit )
 {
 	ExposureBias = 0;
+	bIsRenderingEnabled = true;
 	ToolkitPtr = InToolkit;
 	
 	// create zoom menu
@@ -247,6 +248,15 @@ void STextureEditorViewport::ModifyCheckerboardTextureColors( )
 	}
 }
 
+void STextureEditorViewport::EnableRendering()
+{
+	bIsRenderingEnabled = true;
+}
+
+void STextureEditorViewport::DisableRendering()
+{
+	bIsRenderingEnabled = false;
+}
 
 TSharedPtr<FSceneViewport> STextureEditorViewport::GetViewport( ) const
 {
@@ -274,8 +284,11 @@ TSharedPtr<SScrollBar> STextureEditorViewport::GetHorizontalScrollBar( ) const
 
 void STextureEditorViewport::Tick( const FGeometry& AllottedGeometry, const double InCurrentTime, const float InDeltaTime )
 {
-	Viewport->Invalidate();
-	Viewport->InvalidateDisplay();
+	if (bIsRenderingEnabled)
+	{
+		Viewport->Invalidate();
+		Viewport->InvalidateDisplay();
+	}
 }
 
 
diff --git a/Engine/Source/Editor/TextureEditor/Private/Widgets/STextureEditorViewport.h b/Engine/Source/Editor/TextureEditor/Private/Widgets/STextureEditorViewport.h
--- a/Engine/Source/Editor/TextureEditor/Private/Widgets/STextureEditorViewport.h
+++ b/Engine/Source/Editor/TextureEditor/Private/Widgets/STextureEditorViewport.h
@@ -46,6 +46,12 @@ public:
 		return ExposureBias;
 	}
 
+	/** Enable viewport rendering */
+	void EnableRendering();
+
+	/** Disable viewport rendering */
+	void DisableRendering();
+
 	TSharedPtr<FSceneViewport> GetViewport( ) const;
 	TSharedPtr<SViewport> GetViewportWidget( ) const;
 	TSharedPtr<SScrollBar> GetVerticalScrollBar( ) const;
@@ -127,4 +133,7 @@ private:
 
 	// Horizontal scrollbar.
 	TSharedPtr<SScrollBar> TextureViewportHorizontalScrollBar;
+
+	// Is rendering currently enabled? (disabled when reimporting a texture)
+	bool bIsRenderingEnabled;
 };
diff --git a/Engine/Source/Editor/UnrealEd/Private/Editor.cpp b/Engine/Source/Editor/UnrealEd/Private/Editor.cpp
--- a/Engine/Source/Editor/UnrealEd/Private/Editor.cpp
+++ b/Engine/Source/Editor/UnrealEd/Private/Editor.cpp
@@ -2469,6 +2469,9 @@ bool FReimportManager::CanReimport( UObject* Obj ) const
 
 bool FReimportManager::Reimport( UObject* Obj, bool bAskForNewFileIfMissing )
 {
+	// Warn that were about to reimport, so prep for it
+	PreReimport.Broadcast( Obj );
+
 	bool bSuccess = false;
 	if ( Obj )
 	{
@@ -2590,6 +2593,9 @@ bool FReimportManager::Reimport( UObject* Obj, bool bAskForNewFileIfMissing )
 		}
 	}
 
+	// Let listeners know whether the reimport was successful or not
+	PostReimport.Broadcast( Obj, bSuccess );
+
 	return bSuccess;
 }
 
diff --git a/Engine/Source/Editor/UnrealEd/Private/Toolkits/AssetEditorToolkit.cpp b/Engine/Source/Editor/UnrealEd/Private/Toolkits/AssetEditorToolkit.cpp
--- a/Engine/Source/Editor/UnrealEd/Private/Toolkits/AssetEditorToolkit.cpp
+++ b/Engine/Source/Editor/UnrealEd/Private/Toolkits/AssetEditorToolkit.cpp
@@ -512,14 +512,8 @@ void FAssetEditorToolkit::Reimport_Execute( UObject* EditingObject )
 	// Don't allow user to perform certain actions on objects that aren't actually assets (e.g. Level Script blueprint objects)
 	if( EditingObject != NULL && EditingObject->IsAsset() )
 	{
-		// Warn that were about to reimport, so prep for it
-		PreReimport.Broadcast( EditingObject );
-
 		// Reimport the asset
-		const bool bSuccess = FReimportManager::Instance()->Reimport(EditingObject, ShouldPromptForNewFilesOnReload(*EditingObject));
-
-		// Let listeners know whether the reimport was successful or not
-		PostReimport.Broadcast( EditingObject, bSuccess );
+		FReimportManager::Instance()->Reimport(EditingObject, ShouldPromptForNewFilesOnReload(*EditingObject));
 	}
 }
 
diff --git a/Engine/Source/Editor/UnrealEd/Public/EditorReimportHandler.h b/Engine/Source/Editor/UnrealEd/Public/EditorReimportHandler.h
--- a/Engine/Source/Editor/UnrealEd/Public/EditorReimportHandler.h
+++ b/Engine/Source/Editor/UnrealEd/Public/EditorReimportHandler.h
@@ -70,6 +70,23 @@ public:
 	 */
 	UNREALED_API static FString ResolveImportFilename(const FString& InRelativePath, const UObject* Obj);
 
+	/**
+	 * Gets the delegate that's fired prior to reimporting an asset
+	 *
+	 * @return The event delegate.
+	 */
+	DECLARE_MULTICAST_DELEGATE_OneParam(FPreReimportNotification, UObject*);
+	FPreReimportNotification& OnPreReimport(){ return PreReimport; }
+
+	/**
+	 * Gets the delegate that's fired after to reimporting an asset
+	 *
+	 * @param Whether the reimport was a success
+	 * @return The event delegate.
+	 */
+	DECLARE_MULTICAST_DELEGATE_TwoParams(FPostReimportNotification, UObject*, bool);
+	FPostReimportNotification& OnPostReimport(){ return PostReimport; }
+
 private:
 	/** Opens a file dialog to request a new reimport path */
 	void GetNewReimportPath(UObject* Obj, TArray<FString>& InOutFilenames);
@@ -78,6 +95,12 @@ private:
 	/** Reimport handlers registered with this manager */
 	TArray<FReimportHandler*> Handlers;
 
+	/** Delegate to call before the asset is reimported */
+	FPreReimportNotification PreReimport;
+
+	/** Delegate to call after the asset is reimported */
+	FPostReimportNotification PostReimport;
+
 	/** Constructor */
 	FReimportManager();
 
diff --git a/Engine/Source/Editor/UnrealEd/Public/Toolkits/AssetEditorToolkit.h b/Engine/Source/Editor/UnrealEd/Public/Toolkits/AssetEditorToolkit.h
--- a/Engine/Source/Editor/UnrealEd/Public/Toolkits/AssetEditorToolkit.h
+++ b/Engine/Source/Editor/UnrealEd/Public/Toolkits/AssetEditorToolkit.h
@@ -127,23 +127,6 @@ public:
 	/** True if this actually is editing an asset */
 	bool IsActuallyAnAsset() const;
 
-	/**
-	 * Gets the delegate that's fired prior to reimporting an asset
-	 *
-	 * @return The event delegate.
-	 */
-	DECLARE_MULTICAST_DELEGATE_OneParam(FPreReimportNotification, UObject*);
-	FPreReimportNotification& OnPreReimport(){ return PreReimport; }
-
-	/**
-	 * Gets the delegate that's fired after to reimporting an asset
-	 *
-	 * @param Whether the reimport was a success
-	 * @return The event delegate.
-	 */
-	DECLARE_MULTICAST_DELEGATE_TwoParams(FPostReimportNotification, UObject*, bool);
-	FPostReimportNotification& OnPostReimport(){ return PostReimport; }
-
 protected:
 
 	/**	Returns the single object currently being edited. Asserts if currently editing no object or multiple objects */
@@ -244,12 +227,6 @@ private:
 	/** Whether the buttons on the default toolbar can receive keyboard focus */
 	bool bIsToolbarFocusable;
 
-	/** Delegate to call before the asset is reimported */
-	FPreReimportNotification PreReimport;
-
-	/** Delegate to call after the asset is reimported */
-	FPostReimportNotification PostReimport;
-
 	/**	The tab ids for all the tabs used */
 	static const FName ToolbarTabId;
 };
