Texture Editor: bUseSpecifiedMipLevel now default to false and the number of cinematic mips is ignored by the UI when it's true and Max In-Game no longer takes the editor specified mip setting into account

#ttp 333394 - EDITOR: texture editor bugs (was RE: Flow map action for Photoshop CS5 and above)

#branch UE4

#change Max In-Game dimensions are now calculated seperately as to not confuse the user.

#change Renamed 'Current' to 'Displayed'

#change bUseSpecifiedMipLevel was uninitialized, so could appear as true in the UI without the users knowledge.

#change GetMaxMipLevel û a couple of calls to this were assuming it returned a float, rather than an int32.

#change Cleaned up CachedCombinedLODBias usage and removed some duplicate code.

#change  æNum Cinematic MipsÆ is now also ignored by the UI when æUse specified mipsÆ is enabled.

diff --git a/Engine/Source/Editor/TextureEditor/Private/TextureEditorToolkit.cpp b/Engine/Source/Editor/TextureEditor/Private/TextureEditorToolkit.cpp
--- a/Engine/Source/Editor/TextureEditor/Private/TextureEditorToolkit.cpp
+++ b/Engine/Source/Editor/TextureEditor/Private/TextureEditorToolkit.cpp
@@ -84,6 +84,7 @@ void FTextureEditorToolkit::InitTextureEditor( const EToolkitMode::Type Mode, co
 	PreviewEffectiveTextureWidth = 0;
 	PreviewEffectiveTextureHeight = 0;
 	SpecifiedMipLevel = 0;
+	bUseSpecifiedMipLevel = false;
 
 	SavedCompressionSetting = false;
 
@@ -270,9 +271,6 @@ void FTextureEditorToolkit::PopulateQuickInfo( )
 
 	uint32 ImportedWidth = Texture->Source.GetSizeX();
 	uint32 ImportedHeight = Texture->Source.GetSizeY();
-	uint32 Width, Height;
-
-	Texture->CachedCombinedLODBias = GSystemSettings.TextureLODSettings.CalculateLODBias(Texture, !GetUseSpecifiedMip());
 
 	// If Original Width and Height are 0, use the saved current width and height
 	if (ImportedWidth == 0 && ImportedHeight == 0)
@@ -281,25 +279,28 @@ void FTextureEditorToolkit::PopulateQuickInfo( )
 		ImportedHeight = Texture->GetSurfaceHeight();
 	}
 
-	const int32 MipLevel = FMath::Max( GetMipLevel(), Texture->CachedCombinedLODBias );
+	// In game max bias and dimensions
+	uint32 MaxInGameWidth, MaxInGameHeight;
+	int32 MipLevel = GSystemSettings.TextureLODSettings.CalculateLODBias(Texture);
+	CalculateEffectiveTextureDimensions(MipLevel, MaxInGameWidth, MaxInGameHeight);
+
+	// Editor bias and dimensions (takes user specified mip setting into account)
+	Texture->UpdateCachedLODBias(!GetUseSpecifiedMip());
+	MipLevel = FMath::Max(GetMipLevel(), Texture->GetCachedLODBias());
 	CalculateEffectiveTextureDimensions(MipLevel, PreviewEffectiveTextureWidth, PreviewEffectiveTextureHeight);
-	uint32 MaxInGameWidth = PreviewEffectiveTextureWidth;
-	uint32 MaxInGameHeight = PreviewEffectiveTextureHeight;
 
 	// Cubes are previewed as unwrapped 2D textures.
 	// These have 2x the width of a cube face.
 	PreviewEffectiveTextureWidth *= IsCubeTexture() ? 2 : 1;
 
-	CalculateTextureDimensions(Width, Height);
-
 	FNumberFormattingOptions Options;
 	Options.UseGrouping = false;
 
 	ImportedText->SetText( FText::Format( NSLOCTEXT("TextureEditor", "QuickInfo_Imported", "Imported: {0}x{1}"), FText::AsNumber( ImportedWidth, &Options ), FText::AsNumber( ImportedHeight, &Options ) ) );
-	CurrentText->SetText(  FText::Format( NSLOCTEXT("TextureEditor", "QuickInfo_Current", "Current: {0}x{1}"), FText::AsNumber( FMath::Max((uint32)1, ImportedWidth >> MipLevel), &Options ), FText::AsNumber( FMath::Max((uint32)1, ImportedHeight >> MipLevel), &Options) ) );
+	CurrentText->SetText(  FText::Format( NSLOCTEXT("TextureEditor", "QuickInfo_Displayed", "Displayed: {0}x{1}"), FText::AsNumber( FMath::Max((uint32)1, ImportedWidth >> MipLevel), &Options ), FText::AsNumber( FMath::Max((uint32)1, ImportedHeight >> MipLevel), &Options) ) );
 	MaxInGameText->SetText(FText::Format( NSLOCTEXT("TextureEditor", "QuickInfo_MaxInGame", "Max In-Game: {0}x{1}"), FText::AsNumber( MaxInGameWidth, &Options ), FText::AsNumber( MaxInGameHeight, &Options )));
 	MethodText->SetText(   FText::Format( NSLOCTEXT("TextureEditor", "QuickInfo_Method", "Method: {0}"), Texture->NeverStream ? NSLOCTEXT("TextureEditor", "QuickInfo_MethodNotStreamed", "Not Streamed") : NSLOCTEXT("TextureEditor", "QuickInfo_MethodStreamed", "Streamed")));
-	LODBiasText->SetText(  FText::Format( NSLOCTEXT("TextureEditor", "QuickInfo_LODBias", "Combined LOD Bias: {0}"), FText::AsNumber( Texture->CachedCombinedLODBias )));
+	LODBiasText->SetText(FText::Format(NSLOCTEXT("TextureEditor", "QuickInfo_LODBias", "Combined LOD Bias: {0}"), FText::AsNumber(Texture->GetCachedLODBias())));
 	
 	int32 TextureFormatIndex = PF_MAX;
 	
@@ -962,7 +963,7 @@ TOptional<int32> FTextureEditorToolkit::GetMaxMipLevel()const
 
 bool FTextureEditorToolkit::GetUseSpecifiedMip() const
 {
-	if( GetMaxMipLevel().Get(MIPLEVEL_MAX) > 0.0f )
+	if( GetMaxMipLevel().Get(MIPLEVEL_MAX) > 0 )
 	{
 		if( OnGetUseSpecifiedMipEnabled() )
 		{
@@ -1028,7 +1029,7 @@ bool FTextureEditorToolkit::OnGetUseSpecifiedMipEnabled() const
 {
 	UTextureCube* TextureCube = Cast<UTextureCube>(Texture);
 
-	if( GetMaxMipLevel().Get(MIPLEVEL_MAX) <= 0.0f || TextureCube )
+	if( GetMaxMipLevel().Get(MIPLEVEL_MAX) <= 0 || TextureCube )
 	{
 		return false;
 	}
diff --git a/Engine/Source/Runtime/Engine/Classes/Engine/Texture.h b/Engine/Source/Runtime/Engine/Classes/Engine/Texture.h
--- a/Engine/Source/Runtime/Engine/Classes/Engine/Texture.h
+++ b/Engine/Source/Runtime/Engine/Classes/Engine/Texture.h
@@ -460,10 +460,6 @@ public:
 	UPROPERTY(EditAnywhere, BlueprintReadWrite, Category=LevelOfDetail, meta=(DisplayName="LOD Bias"), AssetRegistrySearchable)
 	int32 LODBias;
 
-	/** Cached combined group and texture LOD bias to use.	*/
-	UPROPERTY(transient)
-	int32 CachedCombinedLODBias;
-
 	/** Number of mip-levels to use for cinematic quality. */
 	UPROPERTY(EditAnywhere, BlueprintReadWrite, Category=LevelOfDetail, AdvancedDisplay)
 	int32 NumCinematicMipLevels;
@@ -484,6 +480,10 @@ public:
 	uint32 bUseCinematicMipLevels:1;
 
 private:
+	/** Cached combined group and texture LOD bias to use.	*/
+	UPROPERTY(transient)
+	int32 CachedCombinedLODBias;
+
 	/** Whether the async resource release process has already been kicked off or not */
 	UPROPERTY(transient)
 	uint32 bAsyncResourceReleaseHasBeenStarted:1;
@@ -603,7 +603,7 @@ public:
 	 */
 	void ForceRebuildPlatformData();
 
-	void UpdateCachedLODBias();
+	ENGINE_API void UpdateCachedLODBias( bool bIncTextureMips = true );
 
 	/**
 	 * Marks platform data as transient. This optionally removes persistent or cached data associated with the platform.
diff --git a/Engine/Source/Runtime/Engine/Private/Texture.cpp b/Engine/Source/Runtime/Engine/Private/Texture.cpp
--- a/Engine/Source/Runtime/Engine/Private/Texture.cpp
+++ b/Engine/Source/Runtime/Engine/Private/Texture.cpp
@@ -287,7 +287,7 @@ void UTexture::PostLoad()
 	if( !IsTemplate() )
 	{
 		// Update cached LOD bias.
-		CachedCombinedLODBias = GSystemSettings.TextureLODSettings.CalculateLODBias( this );
+		UpdateCachedLODBias();
 
 		// The texture will be cached by the cubemap it is contained within on consoles.
 		UTextureCube* CubeMap = Cast<UTextureCube>(GetOuter());
@@ -446,7 +446,7 @@ bool UTexture::ForceUpdateTextureStreaming()
 			UTexture* Texture = *It;
 
 			// Update cached LOD bias.
-			Texture->CachedCombinedLODBias = GSystemSettings.TextureLODSettings.CalculateLODBias( Texture );
+			Texture->UpdateCachedLODBias();
 		}
 
 		// Make sure we iterate over all textures by setting it to high value.
@@ -585,45 +585,19 @@ void FTextureLODSettings::ReadEntry( int32 GroupId, const TCHAR* GroupName, cons
 	}
 }
 
-int32 FTextureLODSettings::CalculateLODBias( UTexture* Texture, bool bIncTextureBias ) const
+int32 FTextureLODSettings::CalculateLODBias( UTexture* Texture, bool bIncTextureMips ) const
 {	
-	// Find LOD group.
 	check( Texture );
-	const FTextureLODGroup& LODGroup = TextureLODGroups[Texture->LODGroup];
-
-	// Calculate maximum number of miplevels.
-	int32 TextureMaxLOD	= FMath::CeilLogTwo( FMath::Trunc( FMath::Max( Texture->GetSurfaceWidth(), Texture->GetSurfaceHeight() ) ) );
-
-	// Calculate LOD bias.
-	int32 UsedLODBias		= Texture->NumCinematicMipLevels + LODGroup.LODBias + (bIncTextureBias ? Texture->LODBias : 0);
-	int32 MinLOD			= LODGroup.MinLODMipCount;
-	int32 MaxLOD			= LODGroup.MaxLODMipCount;
-	int32 WantedMaxLOD	= FMath::Clamp( TextureMaxLOD - UsedLODBias, MinLOD, MaxLOD );
-	WantedMaxLOD		= FMath::Clamp( WantedMaxLOD, 0, TextureMaxLOD );
-	UsedLODBias			= TextureMaxLOD - WantedMaxLOD;
-
-	return UsedLODBias;
+	return CalculateLODBias(Texture->GetSurfaceWidth(), Texture->GetSurfaceHeight(), Texture->LODGroup, (bIncTextureMips ? Texture->LODBias : 0), (bIncTextureMips ? Texture->NumCinematicMipLevels : 0), Texture->MipGenSettings);
 }
 
-/**
-	* Calculates and returns the LOD bias based on the information provided.
-	*
-	* @param	Width						Width of the texture
-	* @param	Height						Height of the texture
-	* @param	LODGroup					Which LOD group the texture belongs to
-	* @param	LODBias						LOD bias to include in the calculation
-	* @param	NumCinematicMipLevels		The texture cinematic mip levels to include in the calculation
-	* @return	LOD bias
-	*/
 int32 FTextureLODSettings::CalculateLODBias( int32 Width, int32 Height, int32 LODGroup, int32 LODBias, int32 NumCinematicMipLevels, TextureMipGenSettings InMipGenSetting ) const
 {	
 	// Find LOD group.
 	const FTextureLODGroup& LODGroupInfo = TextureLODGroups[LODGroup];
 
 	// Test to see if we have no mip generation as in which case the LOD bias will be ignored
-
 	const TextureMipGenSettings FinalMipGenSetting = (InMipGenSetting == TMGS_FromTextureGroup) ? LODGroupInfo.MipGenSettings : InMipGenSetting;
-
 	if ( FinalMipGenSetting == TMGS_NoMipmaps )
 	{
 		return 0;
@@ -633,9 +607,9 @@ int32 FTextureLODSettings::CalculateLODBias( int32 Width, int32 Height, int32 LO
 	int32 TextureMaxLOD	= FMath::CeilLogTwo( FMath::Trunc( FMath::Max( Width, Height ) ) );
 
 	// Calculate LOD bias.
-	int32 UsedLODBias		= LODGroupInfo.LODBias + LODBias + NumCinematicMipLevels;
-	int32 MinLOD			= LODGroupInfo.MinLODMipCount;
-	int32 MaxLOD			= LODGroupInfo.MaxLODMipCount;
+	int32 UsedLODBias	= LODGroupInfo.LODBias + LODBias + NumCinematicMipLevels;
+	int32 MinLOD		= LODGroupInfo.MinLODMipCount;
+	int32 MaxLOD		= LODGroupInfo.MaxLODMipCount;
 	int32 WantedMaxLOD	= FMath::Clamp( TextureMaxLOD - UsedLODBias, MinLOD, MaxLOD );
 	WantedMaxLOD		= FMath::Clamp( WantedMaxLOD, 0, TextureMaxLOD );
 	UsedLODBias			= TextureMaxLOD - WantedMaxLOD;
diff --git a/Engine/Source/Runtime/Engine/Private/Texture2D.cpp b/Engine/Source/Runtime/Engine/Private/Texture2D.cpp
--- a/Engine/Source/Runtime/Engine/Private/Texture2D.cpp
+++ b/Engine/Source/Runtime/Engine/Private/Texture2D.cpp
@@ -522,10 +522,9 @@ FString UTexture2D::GetDesc()
 	uint32 EffectiveSizeX;
 	uint32 EffectiveSizeY;
 
-	// platform dependent
-	CachedCombinedLODBias = uint32( GSystemSettings.TextureLODSettings.CalculateLODBias(this) );
+	UpdateCachedLODBias();
 
-	GSystemSettings.TextureLODSettings.ComputeInGameMaxResolution(CachedCombinedLODBias, *this, EffectiveSizeX, EffectiveSizeY);
+	GSystemSettings.TextureLODSettings.ComputeInGameMaxResolution(GetCachedLODBias(), *this, EffectiveSizeX, EffectiveSizeY);
 
 	return FString::Printf( TEXT("%s %dx%d -> %dx%d[%s]"), 
 		NeverStream ? TEXT("NeverStreamed") : TEXT("Streamed"), 
diff --git a/Engine/Source/Runtime/Engine/Private/TextureDerivedData.cpp b/Engine/Source/Runtime/Engine/Private/TextureDerivedData.cpp
--- a/Engine/Source/Runtime/Engine/Private/TextureDerivedData.cpp
+++ b/Engine/Source/Runtime/Engine/Private/TextureDerivedData.cpp
@@ -1368,9 +1368,9 @@ void UTexture::CachePlatformData(bool bAsyncCache)
 	}
 }
 
-void UTexture::UpdateCachedLODBias()
+void UTexture::UpdateCachedLODBias( bool bIncTextureMips )
 {
-	CachedCombinedLODBias = GSystemSettings.TextureLODSettings.CalculateLODBias( this );
+	CachedCombinedLODBias = GSystemSettings.TextureLODSettings.CalculateLODBias(this, bIncTextureMips);
 }
 
 void UTexture::BeginCachePlatformData()
@@ -1518,7 +1518,6 @@ void UTexture::FinishCachePlatformData()
 	}
 
 	UpdateCachedLODBias();
-
 }
 
 void UTexture::ForceRebuildPlatformData()
diff --git a/Engine/Source/Runtime/Engine/Public/Texture.h b/Engine/Source/Runtime/Engine/Public/Texture.h
--- a/Engine/Source/Runtime/Engine/Public/Texture.h
+++ b/Engine/Source/Runtime/Engine/Public/Texture.h
@@ -900,10 +900,10 @@ struct FTextureLODSettings
 	 * Calculates and returns the LOD bias based on texture LOD group, LOD bias and maximum size.
 	 *
 	 * @param	Texture			Texture object to calculate LOD bias for.
-	 * @param	bIncTextureBias	If true, takes the texures LOD bias into consideration
+	 * @param	bIncTextureMips	If true, takes the textures LOD & Cinematic mips into consideration
 	 * @return	LOD bias
 	 */
-	ENGINE_API int32 CalculateLODBias( UTexture* Texture, bool bIncTextureBias = true ) const;
+	ENGINE_API int32 CalculateLODBias(UTexture* Texture, bool bIncTextureMips = true) const;
 
 	/**
 	 * Calculates and returns the LOD bias based on the information provided.
