TextureEditor: code and documentation cleanup pass

diff --git a/Engine/Source/Editor/TextureEditor/Private/Models/TextureEditorCommands.cpp b/Engine/Source/Editor/TextureEditor/Private/Models/TextureEditorCommands.cpp
--- a/Engine/Source/Editor/TextureEditor/Private/Models/TextureEditorCommands.cpp
+++ b/Engine/Source/Editor/TextureEditor/Private/Models/TextureEditorCommands.cpp
@@ -1,9 +1,5 @@
 // Copyright 1998-2014 Epic Games, Inc. All Rights Reserved.
 
-/*=============================================================================
-	TextureEditorCommands.cpp: Implements the FTextureEditorCommands class.
-=============================================================================*/
-
 #include "TextureEditorPrivatePCH.h"
 
 
diff --git a/Engine/Source/Editor/TextureEditor/Private/Models/TextureEditorCommands.h b/Engine/Source/Editor/TextureEditor/Private/Models/TextureEditorCommands.h
--- a/Engine/Source/Editor/TextureEditor/Private/Models/TextureEditorCommands.h
+++ b/Engine/Source/Editor/TextureEditor/Private/Models/TextureEditorCommands.h
@@ -1,9 +1,5 @@
 // Copyright 1998-2014 Epic Games, Inc. All Rights Reserved.
 
-/*=============================================================================
-	TextureEditorCommands.h: Declares the FTextureEditorCommands class.
-=============================================================================*/
-
 #pragma once
 
 
@@ -24,7 +20,8 @@ public:
 
 public:
 
-	/** Initialize commands */
+	// TCommands interface
+
 	virtual void RegisterCommands() override;
 	
 public:
diff --git a/Engine/Source/Editor/TextureEditor/Private/TextureEditorToolkit.cpp b/Engine/Source/Editor/TextureEditor/Private/TextureEditorToolkit.cpp
--- a/Engine/Source/Editor/TextureEditor/Private/TextureEditorToolkit.cpp
+++ b/Engine/Source/Editor/TextureEditor/Private/TextureEditorToolkit.cpp
@@ -1,13 +1,7 @@
 // Copyright 1998-2014 Epic Games, Inc. All Rights Reserved.
 
-/*=============================================================================
-	TextureEditorToolkit.cpp: Implements the FTextureEditorToolkit class.
-=============================================================================*/
-
 #include "TextureEditorPrivatePCH.h"
-
 #include "Factories.h"
-
 #include "Runtime/Analytics/Analytics/Public/Interfaces/IAnalyticsProvider.h"
 #include "EngineAnalytics.h"
 
@@ -40,19 +34,25 @@
 /* FAssetEditorToolkit interface
  *****************************************************************************/
 
+FString FTextureEditorToolkit::GetDocumentationLink( ) const 
+{
+	return FString(TEXT("Engine/Content/Types/Textures/Properties/Interface"));
+}
+
+
 void FTextureEditorToolkit::RegisterTabSpawners( const TSharedRef<class FTabManager>& TabManager )
 {
 	FAssetEditorToolkit::RegisterTabSpawners(TabManager);
 
 	const IWorkspaceMenuStructure& MenuStructure = WorkspaceMenu::GetMenuStructure();
 
-	TabManager->RegisterTabSpawner( ViewportTabId, FOnSpawnTab::CreateSP(this, &FTextureEditorToolkit::SpawnTab_Viewport) )
-		.SetDisplayName( LOCTEXT("ViewportTab", "Viewport") )
-		.SetGroup( MenuStructure.GetAssetEditorCategory() );
+	TabManager->RegisterTabSpawner(ViewportTabId, FOnSpawnTab::CreateSP(this, &FTextureEditorToolkit::HandleTabSpawnerSpawnViewport))
+		.SetDisplayName(LOCTEXT("ViewportTab", "Viewport"))
+		.SetGroup(MenuStructure.GetAssetEditorCategory());
 	
-	TabManager->RegisterTabSpawner( PropertiesTabId, FOnSpawnTab::CreateSP(this, &FTextureEditorToolkit::SpawnTab_Properties) )
-		.SetDisplayName( LOCTEXT("PropertiesTab", "Details") )
-		.SetGroup( MenuStructure.GetAssetEditorCategory() );
+	TabManager->RegisterTabSpawner(PropertiesTabId, FOnSpawnTab::CreateSP(this, &FTextureEditorToolkit::HandleTabSpawnerSpawnProperties))
+		.SetDisplayName(LOCTEXT("PropertiesTab", "Details") )
+		.SetGroup(MenuStructure.GetAssetEditorCategory());
 }
 
 
@@ -60,15 +60,15 @@
 {
 	FAssetEditorToolkit::UnregisterTabSpawners(TabManager);
 
-	TabManager->UnregisterTabSpawner( ViewportTabId );
-	TabManager->UnregisterTabSpawner( PropertiesTabId );
+	TabManager->UnregisterTabSpawner(ViewportTabId);
+	TabManager->UnregisterTabSpawner(PropertiesTabId);
 }
 
 
 void FTextureEditorToolkit::InitTextureEditor( const EToolkitMode::Type Mode, const TSharedPtr< class IToolkitHost >& InitToolkitHost, UObject* ObjectToEdit )
 {
-	FReimportManager::Instance()->OnPreReimport().AddRaw(this, &FTextureEditorToolkit::OnPreReimport);
-	FReimportManager::Instance()->OnPostReimport().AddRaw(this, &FTextureEditorToolkit::OnPostReimport);
+	FReimportManager::Instance()->OnPreReimport().AddRaw(this, &FTextureEditorToolkit::HandleReimportManagerPreReimport);
+	FReimportManager::Instance()->OnPostReimport().AddRaw(this, &FTextureEditorToolkit::HandleReimportManagerPostReimport);
 
 	Texture = CastChecked<UTexture>(ObjectToEdit);
 
@@ -96,7 +96,6 @@
 	FTextureEditorCommands::Register();
 
 	BindCommands();
-
 	CreateInternalWidgets();
 
 	const TSharedRef<FTabManager::FLayout> StandaloneDefaultLayout = FTabManager::NewLayout("Standalone_TextureEditor_Layout_v3")
@@ -135,12 +134,14 @@
 
 	const bool bCreateDefaultStandaloneMenu = true;
 	const bool bCreateDefaultToolbar = true;
+
 	FAssetEditorToolkit::InitAssetEditor(Mode, InitToolkitHost, TextureEditorAppIdentifier, StandaloneDefaultLayout, bCreateDefaultStandaloneMenu, bCreateDefaultToolbar, ObjectToEdit);
 	
 	ITextureEditorModule* TextureEditorModule = &FModuleManager::LoadModuleChecked<ITextureEditorModule>("TextureEditor");
 	AddMenuExtender(TextureEditorModule->GetMenuExtensibilityManager()->GetAllExtenders(GetToolkitCommands(), GetEditingObjects()));
 
 	ExtendToolBar();
+
 	RegenerateMenusAndToolbars();
 
 	// @todo toolkit world centric editing
@@ -153,12 +154,6 @@
 }
 
 
-bool FTextureEditorToolkit::IsCubeTexture( ) const
-{
-	return (Texture->IsA(UTextureCube::StaticClass()) || Texture->IsA(UTextureRenderTargetCube::StaticClass()));
-}
-
-
 /* ITextureEditorToolkit interface
  *****************************************************************************/
 
@@ -264,6 +259,49 @@
 }
 
 
+bool FTextureEditorToolkit::GetFitToViewport( ) const
+{
+	const UTextureEditorSettings& Settings = *GetDefault<UTextureEditorSettings>();
+	return Settings.FitToViewport;
+}
+
+
+int32 FTextureEditorToolkit::GetMipLevel( ) const
+{
+	return GetUseSpecifiedMip() ? SpecifiedMipLevel : 0;
+}
+
+
+UTexture* FTextureEditorToolkit::GetTexture( ) const
+{
+	return Texture;
+}
+
+
+bool FTextureEditorToolkit::GetUseSpecifiedMip( ) const
+{
+	if (GetMaxMipLevel().Get(MIPLEVEL_MAX) > 0)
+	{
+		if (HandleMipLevelCheckBoxIsEnabled())
+		{
+			return bUseSpecifiedMipLevel;
+		}
+
+		// by default this is on
+		return true; 
+	}
+
+	// disable the widgets if we have no mip maps
+	return false;
+}
+
+
+double FTextureEditorToolkit::GetZoom( ) const
+{
+	return Zoom;
+}
+
+
 void FTextureEditorToolkit::PopulateQuickInfo( )
 {
 	UTexture2D* Texture2D = Cast<UTexture2D>(Texture);
@@ -298,10 +336,10 @@
 	FNumberFormattingOptions Options;
 	Options.UseGrouping = false;
 
-	ImportedText->SetText( FText::Format( NSLOCTEXT("TextureEditor", "QuickInfo_Imported", "Imported: {0}x{1}"), FText::AsNumber( ImportedWidth, &Options ), FText::AsNumber( ImportedHeight, &Options ) ) );
-	CurrentText->SetText(  FText::Format( NSLOCTEXT("TextureEditor", "QuickInfo_Displayed", "Displayed: {0}x{1}"), FText::AsNumber( FMath::Max((uint32)1, ImportedWidth >> MipLevel), &Options ), FText::AsNumber( FMath::Max((uint32)1, ImportedHeight >> MipLevel), &Options) ) );
-	MaxInGameText->SetText(FText::Format( NSLOCTEXT("TextureEditor", "QuickInfo_MaxInGame", "Max In-Game: {0}x{1}"), FText::AsNumber( MaxInGameWidth, &Options ), FText::AsNumber( MaxInGameHeight, &Options )));
-	MethodText->SetText(   FText::Format( NSLOCTEXT("TextureEditor", "QuickInfo_Method", "Method: {0}"), Texture->NeverStream ? NSLOCTEXT("TextureEditor", "QuickInfo_MethodNotStreamed", "Not Streamed") : NSLOCTEXT("TextureEditor", "QuickInfo_MethodStreamed", "Streamed")));
+	ImportedText->SetText(FText::Format( NSLOCTEXT("TextureEditor", "QuickInfo_Imported", "Imported: {0}x{1}"), FText::AsNumber(ImportedWidth, &Options), FText::AsNumber(ImportedHeight, &Options)));
+	CurrentText->SetText(FText::Format( NSLOCTEXT("TextureEditor", "QuickInfo_Displayed", "Displayed: {0}x{1}"), FText::AsNumber(FMath::Max((uint32)1, ImportedWidth >> MipLevel), &Options ), FText::AsNumber(FMath::Max((uint32)1, ImportedHeight>> MipLevel), &Options)));
+	MaxInGameText->SetText(FText::Format( NSLOCTEXT("TextureEditor", "QuickInfo_MaxInGame", "Max In-Game: {0}x{1}"), FText::AsNumber(MaxInGameWidth, &Options), FText::AsNumber(MaxInGameHeight, &Options)));
+	MethodText->SetText(FText::Format( NSLOCTEXT("TextureEditor", "QuickInfo_Method", "Method: {0}"), Texture->NeverStream ? NSLOCTEXT("TextureEditor", "QuickInfo_MethodNotStreamed", "Not Streamed") : NSLOCTEXT("TextureEditor", "QuickInfo_MethodStreamed", "Streamed")));
 	LODBiasText->SetText(FText::Format(NSLOCTEXT("TextureEditor", "QuickInfo_LODBias", "Combined LOD Bias: {0}"), FText::AsNumber(Texture->GetCachedLODBias())));
 	
 	int32 TextureFormatIndex = PF_MAX;
@@ -321,52 +359,62 @@
 
 	if (TextureFormatIndex != PF_MAX)
 	{
-		FormatText->SetText(FText::Format( NSLOCTEXT("TextureEditor", "QuickInfo_Format", "Format: {0}"), FText::FromString( GPixelFormats[TextureFormatIndex].Name )));
+		FormatText->SetText(FText::Format(NSLOCTEXT("TextureEditor", "QuickInfo_Format", "Format: {0}"), FText::FromString(GPixelFormats[TextureFormatIndex].Name)));
 	}
 }
 
 
+void FTextureEditorToolkit::SetFitToViewport( const bool bFitToViewport )
+{
+	UTextureEditorSettings& Settings = *GetMutableDefault<UTextureEditorSettings>();
+	Settings.FitToViewport = bFitToViewport;
+	Settings.PostEditChange();
+}
+
+
+void FTextureEditorToolkit::SetZoom( double ZoomValue )
+{
+	Zoom = FMath::Clamp(ZoomValue, MinZoom, MaxZoom);
+	SetFitToViewport(false);
+}
+
+
+void FTextureEditorToolkit::ZoomIn( )
+{
+	SetZoom(Zoom + ZoomStep);
+}
+
+
+void FTextureEditorToolkit::ZoomOut( )
+{
+	SetZoom(Zoom - ZoomStep);
+}
+
+
 /* IToolkit interface
  *****************************************************************************/
 
-FText FTextureEditorToolkit::GetBaseToolkitName() const
+FText FTextureEditorToolkit::GetBaseToolkitName( ) const
 {
 	return LOCTEXT("AppLabel", "Texture Editor");
 }
 
 
-FString FTextureEditorToolkit::GetWorldCentricTabPrefix() const
+FName FTextureEditorToolkit::GetToolkitFName( ) const
 {
+	return FName("TextureEditor");
+}
+
+
+FLinearColor FTextureEditorToolkit::GetWorldCentricTabColorScale( ) const
+{
+	return FLinearColor(0.3f, 0.2f, 0.5f, 0.5f);
+}
+
+
+FString FTextureEditorToolkit::GetWorldCentricTabPrefix( ) const
+{
 	return LOCTEXT("WorldCentricTabPrefix", "Texture ").ToString();
 }
 
 
-TSharedRef<SDockTab> FTextureEditorToolkit::SpawnTab_Viewport(const FSpawnTabArgs& Args)
-{
-	check( Args.GetTabId() == ViewportTabId );
-
-	return SNew(SDockTab)
-		.Label(LOCTEXT("TextureViewportTitle", "Viewport"))
-		[
-			TextureViewport.ToSharedRef()
-		];
-}
-
-
-TSharedRef<SDockTab> FTextureEditorToolkit::SpawnTab_Properties(const FSpawnTabArgs& Args)
-{
-	check( Args.GetTabId() == PropertiesTabId );
-
-	TSharedRef<SDockTab> SpawnedTab = SNew(SDockTab)
-		.Icon(FEditorStyle::GetBrush("TextureEditor.Tabs.Properties"))
-		.Label(LOCTEXT("TexturePropertiesTitle", "Details"))
-		[
-			TextureProperties.ToSharedRef()
-		];
-
-	PopulateQuickInfo();
-
-	return SpawnedTab;
-}
-
-
@@ -380,10 +428,124 @@
 }
 
 
-/* FTextureEditorToolkit static functions
+/* FEditorUndoClient interface
  *****************************************************************************/
 
-void FTextureEditorToolkit::CreateInternalWidgets()
+void FTextureEditorToolkit::PostUndo( bool bSuccess )
 {
+}
+
+
+void FTextureEditorToolkit::PostRedo( bool bSuccess )
+{
+	PostUndo(bSuccess);
+}
+
+
+/* FTextureEditorToolkit implementation
+ *****************************************************************************/
+
+void FTextureEditorToolkit::BindCommands( )
+{
+	const FTextureEditorCommands& Commands = FTextureEditorCommands::Get();
+
+	ToolkitCommands->MapAction(
+		Commands.RedChannel,
+		FExecuteAction::CreateSP(this, &FTextureEditorToolkit::HandleRedChannelActionExecute),
+		FCanExecuteAction(),
+		FIsActionChecked::CreateSP(this, &FTextureEditorToolkit::HandleRedChannelActionIsChecked));
+
+	ToolkitCommands->MapAction(
+		Commands.GreenChannel,
+		FExecuteAction::CreateSP(this, &FTextureEditorToolkit::HandleGreenChannelActionExecute),
+		FCanExecuteAction(),
+		FIsActionChecked::CreateSP(this, &FTextureEditorToolkit::HandleGreenChannelActionIsChecked));
+
+	ToolkitCommands->MapAction(
+		Commands.BlueChannel,
+		FExecuteAction::CreateSP(this, &FTextureEditorToolkit::HandleBlueChannelActionExecute),
+		FCanExecuteAction(),
+		FIsActionChecked::CreateSP(this, &FTextureEditorToolkit::HandleBlueChannelActionIsChecked));
+
+	ToolkitCommands->MapAction(
+		Commands.AlphaChannel,
+		FExecuteAction::CreateSP(this, &FTextureEditorToolkit::HandleAlphaChannelActionExecute),
+		FCanExecuteAction(),
+		FIsActionChecked::CreateSP(this, &FTextureEditorToolkit::HandleAlphaChannelActionIsChecked));
+
+	ToolkitCommands->MapAction(
+		Commands.Saturation,
+		FExecuteAction::CreateSP(this, &FTextureEditorToolkit::HandleSaturationChannelActionExecute),
+		FCanExecuteAction(),
+		FIsActionChecked::CreateSP(this, &FTextureEditorToolkit::HandleSaturationChannelActionIsChecked));
+
+	ToolkitCommands->MapAction(
+		Commands.FitToViewport,
+		FExecuteAction::CreateSP(this, &FTextureEditorToolkit::HandleFitToViewportActionExecute),
+		FCanExecuteAction(),
+		FIsActionChecked::CreateSP(this, &FTextureEditorToolkit::HandleFitToViewportActionIsChecked));
+
+	ToolkitCommands->MapAction(
+		Commands.CheckeredBackground,
+		FExecuteAction::CreateSP(this, &FTextureEditorToolkit::HandleCheckeredBackgroundActionExecute, TextureEditorBackground_Checkered),
+		FCanExecuteAction(),
+		FIsActionChecked::CreateSP(this, &FTextureEditorToolkit::HandleCheckeredBackgroundActionIsChecked, TextureEditorBackground_Checkered));
+
+	ToolkitCommands->MapAction(
+		Commands.CheckeredBackgroundFill,
+		FExecuteAction::CreateSP(this, &FTextureEditorToolkit::HandleCheckeredBackgroundActionExecute, TextureEditorBackground_CheckeredFill),
+		FCanExecuteAction(),
+		FIsActionChecked::CreateSP(this, &FTextureEditorToolkit::HandleCheckeredBackgroundActionIsChecked, TextureEditorBackground_CheckeredFill));
+
+	ToolkitCommands->MapAction(
+		Commands.SolidBackground,
+		FExecuteAction::CreateSP(this, &FTextureEditorToolkit::HandleCheckeredBackgroundActionExecute, TextureEditorBackground_SolidColor),
+		FCanExecuteAction(),
+		FIsActionChecked::CreateSP(this, &FTextureEditorToolkit::HandleCheckeredBackgroundActionIsChecked, TextureEditorBackground_SolidColor));
+
+	ToolkitCommands->MapAction(
+		Commands.TextureBorder,
+		FExecuteAction::CreateSP(this, &FTextureEditorToolkit::HandleTextureBorderActionExecute),
+		FCanExecuteAction(),
+		FIsActionChecked::CreateSP(this, &FTextureEditorToolkit::HandleTextureBorderActionIsChecked));
+
+	ToolkitCommands->MapAction(
+		Commands.CompressNow,
+		FExecuteAction::CreateSP(this, &FTextureEditorToolkit::HandleCompressNowActionExecute),
+		FCanExecuteAction::CreateSP(this, &FTextureEditorToolkit::HandleCompressNowActionCanExecute));
+
+	ToolkitCommands->MapAction(
+		Commands.Reimport,
+		FExecuteAction::CreateSP(this, &FTextureEditorToolkit::HandleReimportActionExecute),
+		FCanExecuteAction::CreateSP(this, &FTextureEditorToolkit::HandleReimportActionCanExecute));
+
+	ToolkitCommands->MapAction(
+		Commands.Settings,
+		FExecuteAction::CreateSP(this, &FTextureEditorToolkit::HandleSettingsActionExecute));
+}
+
+
+TSharedRef<SWidget> FTextureEditorToolkit::BuildTexturePropertiesWidget( )
+{
+	FDetailsViewArgs Args;
+	Args.bHideSelectionTip = true;
+
+	FPropertyEditorModule& PropertyModule = FModuleManager::LoadModuleChecked<FPropertyEditorModule>("PropertyEditor");
+	TexturePropertiesWidget = PropertyModule.CreateDetailView(Args);
+	TexturePropertiesWidget->SetObject(Texture);
+
+	return TexturePropertiesWidget.ToSharedRef();
+}
+
+
+void FTextureEditorToolkit::CalculateEffectiveTextureDimensions( int32 LODBias, uint32& EffectiveTextureWidth, uint32& EffectiveTextureHeight )
+{
+	//Calculate in-game max resolution and store in EffectiveTextureWidth, EffectiveTextureHeight
+	GSystemSettings.TextureLODSettings.ComputeInGameMaxResolution(LODBias, *Texture, (uint32 &)EffectiveTextureWidth, (uint32 &)EffectiveTextureHeight);
+}
+
+
+void FTextureEditorToolkit::CreateInternalWidgets( )
+{
 	TextureViewport = SNew(STextureEditorViewport, SharedThis(this));
 
@@ -472,21 +634,8 @@
 }
 
 
-TSharedRef<SWidget> FTextureEditorToolkit::BuildTexturePropertiesWidget()
-{
-	FDetailsViewArgs Args;
-	Args.bHideSelectionTip = true;
-
-	FPropertyEditorModule& PropertyModule = FModuleManager::LoadModuleChecked<FPropertyEditorModule>("PropertyEditor");
-	TexturePropertiesWidget = PropertyModule.CreateDetailView(Args);
-	TexturePropertiesWidget->SetObject( Texture );
-
-	return TexturePropertiesWidget.ToSharedRef();
-}
-
-
 BEGIN_SLATE_FUNCTION_BUILD_OPTIMIZATION
-void FTextureEditorToolkit::ExtendToolBar()
+void FTextureEditorToolkit::ExtendToolBar( )
 {
 	struct Local
 	{
@@ -498,7 +647,7 @@
 				ToolbarBuilder.AddToolBarButton(FTextureEditorCommands::Get().Reimport);
 			}
 			ToolbarBuilder.EndSection();
-	
+
 			ToolbarBuilder.BeginSection("TextureMipAndExposure");
 			{
 				ToolbarBuilder.AddWidget(LODControl);
@@ -506,82 +655,89 @@
 			ToolbarBuilder.EndSection();
 		}
 	};
-	
-	TSharedRef<SWidget> LODControl =
-		SNew(SBox)
+
+	TSharedRef<SWidget> LODControl = SNew(SBox)
 		.WidthOverride(240.0f)
 		[
 			SNew(SHorizontalBox)
-			+SHorizontalBox::Slot()
-			.FillWidth(1.0f)
-			.MaxWidth(240.0f)
-			.Padding(0.0f,0.0f,0.0f,0.0f)
-			.VAlign(VAlign_Center)
-			[
-				// Mip and exposure controls
-				SNew(SHorizontalBox)
-				+SHorizontalBox::Slot()
-				.Padding(4.0f,0.0f,4.0f,0.0f)
-				.AutoWidth()
+
+			+ SHorizontalBox::Slot()
+				.FillWidth(1.0f)
+				.MaxWidth(240.0f)
+				.Padding(0.0f, 0.0f, 0.0f, 0.0f)
+				.VAlign(VAlign_Center)
 				[
+					// Mip and exposure controls
 					SNew(SHorizontalBox)
-					+SHorizontalBox::Slot()
-					.VAlign(VAlign_Center)
-					.AutoWidth()
-					[
-						SNew( SCheckBox )
-						.IsChecked( this, &FTextureEditorToolkit::OnGetUseSpecifiedMip )
-						.IsEnabled( this, &FTextureEditorToolkit::OnGetUseSpecifiedMipEnabled )
-						.OnCheckStateChanged( this, &FTextureEditorToolkit::OnUseSpecifiedMipChanged )
-					]
-				]
-				+SHorizontalBox::Slot()
-					.Padding(4.0f,0.0f,4.0f,0.0f)
-					.FillWidth(1)
-					[
-						SNew(SHorizontalBox)
-						+SHorizontalBox::Slot()
-						.Padding(0.0f,0.0,4.0,0.0)
+
+					+ SHorizontalBox::Slot()
+						.Padding(4.0f, 0.0f, 4.0f, 0.0f)
 						.AutoWidth()
-						.VAlign(VAlign_Center)
 						[
-							SNew(STextBlock)
-							.Text(NSLOCTEXT("TextureEditor", "MipLevel", "Mip Level: "))
+							SNew(SHorizontalBox)
+
+							+ SHorizontalBox::Slot()
+								.VAlign(VAlign_Center)
+								.AutoWidth()
+								[
+									SNew(SCheckBox)
+										.IsChecked(this, &FTextureEditorToolkit::HandleMipLevelCheckBoxIsChecked)
+										.IsEnabled(this, &FTextureEditorToolkit::HandleMipLevelCheckBoxIsEnabled)
+										.OnCheckStateChanged(this, &FTextureEditorToolkit::HandleMipLevelCheckBoxCheckedStateChanged)
+								]
 						]
-						+SHorizontalBox::Slot()
-							.VAlign(VAlign_Center)
-							.FillWidth(1.0f)
-							[
-								SNew(SNumericEntryBox<int32>)
-								.AllowSpin(true)
-								.MinSliderValue(MIPLEVEL_MIN)
-								.MaxSliderValue(this, &FTextureEditorToolkit::GetMaxMipLevel)
-								.Value(this, &FTextureEditorToolkit::GetWidgetMipLevel)
-								.OnValueChanged(this, &FTextureEditorToolkit::OnMipLevelChanged)
-								.IsEnabled(this, &FTextureEditorToolkit::GetUseSpecifiedMip)
-							]
-						+SHorizontalBox::Slot()
-							.AutoWidth()
-							.VAlign(VAlign_Center)
-							.Padding(2.0f)
-							[
-								SNew(SButton)
-								.Text(NSLOCTEXT("TextureEditor", "MipMinus", "-"))
-								.OnClicked(this, &FTextureEditorToolkit::OnMipLevelDown)
-								.IsEnabled(this, &FTextureEditorToolkit::GetUseSpecifiedMip)
-							]
-						+SHorizontalBox::Slot()
-							.AutoWidth()
-							.VAlign(VAlign_Center)
-							.Padding(2.0f)
-							[
-								SNew(SButton)
-								.Text(NSLOCTEXT("TextureEditor", "MipPlus", "+"))
-								.OnClicked(this, &FTextureEditorToolkit::OnMipLevelUp)
-								.IsEnabled(this, &FTextureEditorToolkit::GetUseSpecifiedMip)
-							]
-					]
-			]
+
+					+ SHorizontalBox::Slot()
+						.Padding(4.0f, 0.0f, 4.0f, 0.0f)
+						.FillWidth(1)
+						[
+							SNew(SHorizontalBox)
+
+							+ SHorizontalBox::Slot()
+								.Padding(0.0f, 0.0, 4.0, 0.0)
+								.AutoWidth()
+								.VAlign(VAlign_Center)
+								[
+									SNew(STextBlock)
+										.Text(NSLOCTEXT("TextureEditor", "MipLevel", "Mip Level: "))
+								]
+
+							+ SHorizontalBox::Slot()
+								.VAlign(VAlign_Center)
+								.FillWidth(1.0f)
+								[
+									SNew(SNumericEntryBox<int32>)
+										.AllowSpin(true)
+										.MinSliderValue(MIPLEVEL_MIN)
+										.MaxSliderValue(this, &FTextureEditorToolkit::GetMaxMipLevel)
+										.Value(this, &FTextureEditorToolkit::HandleMipLevelEntryBoxValue)
+										.OnValueChanged(this, &FTextureEditorToolkit::HandleMipLevelEntryBoxChanged)
+										.IsEnabled(this, &FTextureEditorToolkit::GetUseSpecifiedMip)
+								]
+
+							+ SHorizontalBox::Slot()
+								.AutoWidth()
+								.VAlign(VAlign_Center)
+								.Padding(2.0f)
+								[
+									SNew(SButton)
+										.Text(NSLOCTEXT("TextureEditor", "MipMinus", "-"))
+										.OnClicked(this, &FTextureEditorToolkit::HandleMipMapMinusButtonClicked)
+										.IsEnabled(this, &FTextureEditorToolkit::GetUseSpecifiedMip)
+								]
+
+							+ SHorizontalBox::Slot()
+								.AutoWidth()
+								.VAlign(VAlign_Center)
+								.Padding(2.0f)
+								[
+									SNew(SButton)
+										.Text(NSLOCTEXT("TextureEditor", "MipPlus", "+"))
+										.OnClicked(this, &FTextureEditorToolkit::HandleMipMapPlusButtonClicked)
+										.IsEnabled(this, &FTextureEditorToolkit::GetUseSpecifiedMip)
+								]
+						]
+				]
 		];
 
 	TSharedPtr<FExtender> ToolbarExtender = MakeShareable(new FExtender);
@@ -590,147 +746,58 @@
 		"Asset",
 		EExtensionHook::After,
 		GetToolkitCommands(),
-		FToolBarExtensionDelegate::CreateStatic( &Local::FillToolbar, GetToolkitCommands(), LODControl )
+		FToolBarExtensionDelegate::CreateStatic(&Local::FillToolbar, GetToolkitCommands(), LODControl)
 	);
-	
+
 	AddToolbarExtender(ToolbarExtender);
 
-	ITextureEditorModule* TextureEditorModule = &FModuleManager::LoadModuleChecked<ITextureEditorModule>( "TextureEditor" );
+	ITextureEditorModule* TextureEditorModule = &FModuleManager::LoadModuleChecked<ITextureEditorModule>("TextureEditor");
 	AddToolbarExtender(TextureEditorModule->GetToolBarExtensibilityManager()->GetAllExtenders(GetToolkitCommands(), GetEditingObjects()));
 }
 END_SLATE_FUNCTION_BUILD_OPTIMIZATION
 
 
-void FTextureEditorToolkit::BindCommands()
+TOptional<int32> FTextureEditorToolkit::GetMaxMipLevel( ) const
 {
-	const FTextureEditorCommands& Commands = FTextureEditorCommands::Get();
+	const UTexture2D* Texture2D = Cast<UTexture2D>(Texture);
+	const UTextureCube* TextureCube = Cast<UTextureCube>(Texture);
+	const UTextureRenderTargetCube* RTTextureCube = Cast<UTextureRenderTargetCube>(Texture);
+	const UTextureRenderTarget2D* RTTexture2D = Cast<UTextureRenderTarget2D>(Texture);
 
-	ToolkitCommands->MapAction(
-		Commands.RedChannel,
-		FExecuteAction::CreateSP(this, &FTextureEditorToolkit::OnRedChannel),
-		FCanExecuteAction(),
-		FIsActionChecked::CreateSP(this, &FTextureEditorToolkit::IsRedChannelChecked));
+	if (Texture2D)
+	{
+		return Texture2D->GetNumMips() - 1;
+	}
 
-	ToolkitCommands->MapAction(
-		Commands.GreenChannel,
-		FExecuteAction::CreateSP(this, &FTextureEditorToolkit::OnGreenChannel),
-		FCanExecuteAction(),
-		FIsActionChecked::CreateSP(this, &FTextureEditorToolkit::IsGreenChannelChecked));
+	if (TextureCube)
+	{
+		return TextureCube->GetNumMips() - 1;
+	}
 
-	ToolkitCommands->MapAction(
-		Commands.BlueChannel,
-		FExecuteAction::CreateSP(this, &FTextureEditorToolkit::OnBlueChannel),
-		FCanExecuteAction(),
-		FIsActionChecked::CreateSP(this, &FTextureEditorToolkit::IsBlueChannelChecked));
+	if (RTTextureCube)
+	{
+		return RTTextureCube->GetNumMips() - 1;
+	}
 
-	ToolkitCommands->MapAction(
-		Commands.AlphaChannel,
-		FExecuteAction::CreateSP(this, &FTextureEditorToolkit::OnAlphaChannel),
-		FCanExecuteAction(),
-		FIsActionChecked::CreateSP(this, &FTextureEditorToolkit::IsAlphaChannelChecked));
+	if (RTTexture2D)
+	{
+		return RTTexture2D->GetNumMips() - 1;
+	}
 
-	ToolkitCommands->MapAction(
-		Commands.Saturation,
-		FExecuteAction::CreateSP(this, &FTextureEditorToolkit::OnSaturation),
-		FCanExecuteAction(),
-		FIsActionChecked::CreateSP(this, &FTextureEditorToolkit::IsSaturationChecked));
-
-	ToolkitCommands->MapAction(
-		Commands.FitToViewport,
-		FExecuteAction::CreateSP(this, &FTextureEditorToolkit::OnFitToViewport),
-		FCanExecuteAction(),
-		FIsActionChecked::CreateSP(this, &FTextureEditorToolkit::IsFitToViewportChecked));
-
-	ToolkitCommands->MapAction(
-		Commands.CheckeredBackground,
-		FExecuteAction::CreateSP(this, &FTextureEditorToolkit::HandleBackgroundActionExecute, TextureEditorBackground_Checkered),
-		FCanExecuteAction(),
-		FIsActionChecked::CreateSP(this, &FTextureEditorToolkit::HandleBackgroundActionIsChecked, TextureEditorBackground_Checkered));
-
-	ToolkitCommands->MapAction(
-		Commands.CheckeredBackgroundFill,
-		FExecuteAction::CreateSP(this, &FTextureEditorToolkit::HandleBackgroundActionExecute, TextureEditorBackground_CheckeredFill),
-		FCanExecuteAction(),
-		FIsActionChecked::CreateSP(this, &FTextureEditorToolkit::HandleBackgroundActionIsChecked, TextureEditorBackground_CheckeredFill));
-
-	ToolkitCommands->MapAction(
-		Commands.SolidBackground,
-		FExecuteAction::CreateSP(this, &FTextureEditorToolkit::HandleBackgroundActionExecute, TextureEditorBackground_SolidColor),
-		FCanExecuteAction(),
-		FIsActionChecked::CreateSP(this, &FTextureEditorToolkit::HandleBackgroundActionIsChecked, TextureEditorBackground_SolidColor));
-
-	ToolkitCommands->MapAction(
-		Commands.TextureBorder,
-		FExecuteAction::CreateSP(this, &FTextureEditorToolkit::OnShowBorder),
-		FCanExecuteAction(),
-		FIsActionChecked::CreateSP(this, &FTextureEditorToolkit::IsShowBorderChecked));
-
-	ToolkitCommands->MapAction(
-		Commands.CompressNow,
-		FExecuteAction::CreateSP(this, &FTextureEditorToolkit::OnCompressNow),
-		FCanExecuteAction::CreateSP(this, &FTextureEditorToolkit::OnCompressNowEnabled));
-
-	ToolkitCommands->MapAction(
-		Commands.Reimport,
-		FExecuteAction::CreateSP(this, &FTextureEditorToolkit::OnReimport),
-		FCanExecuteAction::CreateSP(this, &FTextureEditorToolkit::OnReimportEnabled));
-
-	ToolkitCommands->MapAction(
-		Commands.Settings,
-		FExecuteAction::CreateSP(this, &FTextureEditorToolkit::HandleSettingsActionExecute));
+	return MIPLEVEL_MAX;
 }
 
 
-void FTextureEditorToolkit::CalculateEffectiveTextureDimensions(int32 LODBias, uint32& EffectiveTextureWidth, uint32& EffectiveTextureHeight)
+bool FTextureEditorToolkit::IsCubeTexture( ) const
 {
-	//Calculate in-game max resolution and store in EffectiveTextureWidth, EffectiveTextureHeight
-	GSystemSettings.TextureLODSettings.ComputeInGameMaxResolution(LODBias, *Texture, (uint32 &)EffectiveTextureWidth, (uint32 &)EffectiveTextureHeight);
+	return (Texture->IsA(UTextureCube::StaticClass()) || Texture->IsA(UTextureRenderTargetCube::StaticClass()));
 }
 
 
-void FTextureEditorToolkit::OnRedChannel()
-{
-	bIsRedChannel = !bIsRedChannel;
-}
+/* FTextureEditorToolkit callbacks
+ *****************************************************************************/
 
-
-bool FTextureEditorToolkit::IsRedChannelChecked() const
-{
-	return bIsRedChannel;
-}
-
-
-void FTextureEditorToolkit::OnGreenChannel()
-{
-	bIsGreenChannel = !bIsGreenChannel;
-}
-
-
-bool FTextureEditorToolkit::IsGreenChannelChecked() const
-{
-	return bIsGreenChannel;
-}
-
-
-void FTextureEditorToolkit::OnBlueChannel()
-{
-	bIsBlueChannel = !bIsBlueChannel;
-}
-
-
-bool FTextureEditorToolkit::IsBlueChannelChecked() const
-{
-	return bIsBlueChannel;
-}
-
-
-void FTextureEditorToolkit::OnAlphaChannel()
-{
-	bIsAlphaChannel = !bIsAlphaChannel;
-}
-
-
-bool FTextureEditorToolkit::OnAlphaChannelCanExecute() const
+bool FTextureEditorToolkit::HandleAlphaChannelActionCanExecute( ) const
 {
 	const UTexture2D* Texture2D = Cast<UTexture2D>(Texture);
 
@@ -743,37 +810,31 @@
 }
 
 
-bool FTextureEditorToolkit::IsAlphaChannelChecked() const
+void FTextureEditorToolkit::HandleAlphaChannelActionExecute( )
 {
+	bIsAlphaChannel = !bIsAlphaChannel;
+}
+
+
+bool FTextureEditorToolkit::HandleAlphaChannelActionIsChecked( ) const
+{
 	return bIsAlphaChannel;
 }
 
 
-void FTextureEditorToolkit::OnSaturation()
+void FTextureEditorToolkit::HandleBlueChannelActionExecute( )
 {
-	bIsSaturation = !bIsSaturation;
+	 bIsBlueChannel = !bIsBlueChannel;
 }
 
 
-bool FTextureEditorToolkit::IsSaturationChecked() const
+bool FTextureEditorToolkit::HandleBlueChannelActionIsChecked( ) const
 {
-	return bIsSaturation;
+	return bIsBlueChannel;
 }
 
 
-void FTextureEditorToolkit::OnFitToViewport()
-{
-	ToggleFitToViewport();
-}
-
-
-bool FTextureEditorToolkit::IsFitToViewportChecked() const
-{
-	return GetFitToViewport();
-}
-
-
-void FTextureEditorToolkit::HandleBackgroundActionExecute( ETextureEditorBackgrounds Background )
+void FTextureEditorToolkit::HandleCheckeredBackgroundActionExecute( ETextureEditorBackgrounds Background )
 {
 	UTextureEditorSettings& Settings = *GetMutableDefault<UTextureEditorSettings>();
 	Settings.Background = Background;
@@ -781,7 +842,7 @@
 }
 
 
-bool FTextureEditorToolkit::HandleBackgroundActionIsChecked( ETextureEditorBackgrounds Background )
+bool FTextureEditorToolkit::HandleCheckeredBackgroundActionIsChecked( ETextureEditorBackgrounds Background )
 {
 	const UTextureEditorSettings& Settings = *GetDefault<UTextureEditorSettings>();
 
@@ -789,23 +850,143 @@
 }
 
 
-void FTextureEditorToolkit::OnShowBorder()
+void FTextureEditorToolkit::HandleCompressNowActionExecute( )
 {
-	UTextureEditorSettings& Settings = *GetMutableDefault<UTextureEditorSettings>();
-	Settings.TextureBorderEnabled = !Settings.TextureBorderEnabled;
-	Settings.PostEditChange();
+	GWarn->BeginSlowTask(NSLOCTEXT("TextureEditor", "CompressNow", "Compressing 1 Textures that have Defer Compression set"), true);
+
+	if (Texture->DeferCompression)
+	{
+		// turn off deferred compression and compress the texture
+		Texture->DeferCompression = false;
+		Texture->Source.Compress();
+		Texture->PostEditChange();
+
+		PopulateQuickInfo();
+	}
+
+	GWarn->EndSlowTask();
 }
 
 
-bool FTextureEditorToolkit::IsShowBorderChecked() const
+bool FTextureEditorToolkit::HandleCompressNowActionCanExecute( ) const
 {
-	const UTextureEditorSettings& Settings = *GetDefault<UTextureEditorSettings>();
-
-	return Settings.TextureBorderEnabled;
+	return (Texture->DeferCompression != 0);
 }
 
 
-void FTextureEditorToolkit::OnPreReimport(UObject* InObject)
+void FTextureEditorToolkit::HandleFitToViewportActionExecute( )
+{
+	ToggleFitToViewport();
+}
+
+
+bool FTextureEditorToolkit::HandleFitToViewportActionIsChecked( ) const
+{
+	return GetFitToViewport();
+}
+
+
+void FTextureEditorToolkit::HandleGreenChannelActionExecute( )
+{
+	 bIsGreenChannel = !bIsGreenChannel;
+}
+
+
+bool FTextureEditorToolkit::HandleGreenChannelActionIsChecked( ) const
+{
+	return bIsGreenChannel;
+}
+
+
+void FTextureEditorToolkit::HandleMipLevelCheckBoxCheckedStateChanged( ESlateCheckBoxState::Type InNewState )
+{
+	bUseSpecifiedMipLevel = InNewState == ESlateCheckBoxState::Checked;
+}
+
+
+ESlateCheckBoxState::Type FTextureEditorToolkit::HandleMipLevelCheckBoxIsChecked( ) const
+{
+	return GetUseSpecifiedMip() ? ESlateCheckBoxState::Checked : ESlateCheckBoxState::Unchecked;
+}
+
+
+bool FTextureEditorToolkit::HandleMipLevelCheckBoxIsEnabled( ) const
+{
+	UTextureCube* TextureCube = Cast<UTextureCube>(Texture);
+
+	if (GetMaxMipLevel().Get(MIPLEVEL_MAX) <= 0 || TextureCube)
+	{
+		return false;
+	}
+
+	return true;
+}
+
+
+void FTextureEditorToolkit::HandleMipLevelEntryBoxChanged( int32 NewMipLevel )
+{
+	SpecifiedMipLevel = FMath::Clamp<int32>(NewMipLevel, MIPLEVEL_MIN, GetMaxMipLevel().Get(MIPLEVEL_MAX));
+}
+
+
+TOptional<int32> FTextureEditorToolkit::HandleMipLevelEntryBoxValue( ) const
+{
+	return SpecifiedMipLevel;
+}
+
+
+FReply FTextureEditorToolkit::HandleMipMapMinusButtonClicked( )
+{
+	if (SpecifiedMipLevel > MIPLEVEL_MIN)
+	{
+		--SpecifiedMipLevel;
+	}
+
+	return FReply::Handled();
+}
+
+
+FReply FTextureEditorToolkit::HandleMipMapPlusButtonClicked( )
+{
+	if (SpecifiedMipLevel < GetMaxMipLevel().Get(MIPLEVEL_MAX))
+	{
+		++SpecifiedMipLevel;
+	}
+
+	return FReply::Handled();
+}
+
+
+void FTextureEditorToolkit::HandleRedChannelActionExecute( )
+{
+	bIsRedChannel = !bIsRedChannel;
+}
+
+
+bool FTextureEditorToolkit::HandleRedChannelActionIsChecked( ) const
+{
+	return bIsRedChannel;
+}
+
+
+bool FTextureEditorToolkit::HandleReimportActionCanExecute( ) const
+{
+	if (Texture->IsA<ULightMapTexture2D>() || Texture->IsA<UShadowMapTexture2D>())
+	{
+		return false;
+	}
+
+	return true;
+}
+
+
+void FTextureEditorToolkit::HandleReimportActionExecute( )
+{
+	FReimportManager::Instance()->Reimport(Texture, /*bAskForNewFileIfMissing=*/true);
+}
+
+
+void FTextureEditorToolkit::HandleReimportManagerPostReimport( UObject* InObject, bool bSuccess )
 {
 	// Ignore if this is regarding a different object
 	if (InObject != Texture)
@@ -813,6 +994,32 @@
 		return;
 	}
 
+	if (bSuccess)
+	{
+		if (FEngineAnalytics::IsAvailable())
+		{
+			FEngineAnalytics::GetProvider().RecordEvent(TEXT("Editor.Texture.ReimportedViaEditor"));
+		}
+	}
+	else
+	{
+		// Failed, restore the compression flag
+		Texture->DeferCompression = SavedCompressionSetting;
+	}
+
+	// Re-enable viewport rendering now that the texture should be in a known state again
+	TextureViewport->EnableRendering();
+}
+
+
+void FTextureEditorToolkit::HandleReimportManagerPreReimport( UObject* InObject )
+{
+	// Ignore if this is regarding a different object
+	if (InObject != Texture)
+	{
+		return;
+	}
+
 	// Prevent the texture from being compressed immediately, so the user can see the results
 	SavedCompressionSetting = Texture->DeferCompression;
 	Texture->DeferCompression = true;
@@ -836,220 +1043,66 @@
 }
 
 
-void FTextureEditorToolkit::OnPostReimport(UObject* InObject, bool bSuccess)
+void FTextureEditorToolkit::HandleSaturationChannelActionExecute( )
 {
-	// Ignore if this is regarding a different object
-	if (InObject != Texture)
-	{
-		return;
-	}
-
-	if (bSuccess)
-	{
-		if (FEngineAnalytics::IsAvailable())
-		{
-			FEngineAnalytics::GetProvider().RecordEvent(TEXT("Editor.Texture.ReimportedViaEditor"));
-		}
-	}
-	else
-	{
-		// Failed, restore the compression flag
-		Texture->DeferCompression = SavedCompressionSetting;
-	}
-
-	// Re-enable viewport rendering now that the texture should be in a known state again
-	TextureViewport->EnableRendering();
+	bIsSaturation = !bIsSaturation;
 }
 
 
-void FTextureEditorToolkit::OnCompressNow()
+bool FTextureEditorToolkit::HandleSaturationChannelActionIsChecked( ) const
 {
-	GWarn->BeginSlowTask( NSLOCTEXT("TextureEditor", "CompressNow", "Compressing 1 Textures that have Defer Compression set"), true);
-
-	if (Texture->DeferCompression)
-	{
-		// Turn off defer compression and compress the texture
-		Texture->DeferCompression = false;
-		Texture->Source.Compress();
-		Texture->PostEditChange();
-		PopulateQuickInfo();
-	}
-
-	GWarn->EndSlowTask();
+	return bIsSaturation;
 }
 
 
-bool FTextureEditorToolkit::OnCompressNowEnabled() const
+void FTextureEditorToolkit::HandleSettingsActionExecute( )
 {
-	return Texture->DeferCompression != 0;
-}
-
-
-void FTextureEditorToolkit::OnReimport()
-{
-	FReimportManager::Instance()->Reimport(Texture, /*bAskForNewFileIfMissing=*/true);
-}
-
-bool FTextureEditorToolkit::OnReimportEnabled() const
-{
-	if ( Texture->IsA<ULightMapTexture2D>() || Texture->IsA<UShadowMapTexture2D>() )
-	{
-		return false;
-	}
-	return true;
-}
-
-void FTextureEditorToolkit::HandleSettingsActionExecute()
-{
 	FModuleManager::LoadModuleChecked<ISettingsModule>("Settings").ShowViewer("Editor", "ContentEditors", "TextureEditor");
 }
 
 
-TOptional<int32> FTextureEditorToolkit::GetWidgetMipLevel() const
+TSharedRef<SDockTab> FTextureEditorToolkit::HandleTabSpawnerSpawnProperties( const FSpawnTabArgs& Args )
 {
-	return SpecifiedMipLevel;
+	check(Args.GetTabId() == PropertiesTabId);
+
+	TSharedRef<SDockTab> SpawnedTab = SNew(SDockTab)
+		.Icon(FEditorStyle::GetBrush("TextureEditor.Tabs.Properties"))
+		.Label(LOCTEXT("TexturePropertiesTitle", "Details"))
+		[
+			TextureProperties.ToSharedRef()
+		];
+
+	PopulateQuickInfo();
+
+	return SpawnedTab;
 }
 
 
-void FTextureEditorToolkit::OnMipLevelChanged(int32 NewMipLevel)
+TSharedRef<SDockTab> FTextureEditorToolkit::HandleTabSpawnerSpawnViewport( const FSpawnTabArgs& Args )
 {
-	SpecifiedMipLevel = FMath::Clamp<int32>( NewMipLevel, MIPLEVEL_MIN, GetMaxMipLevel().Get(MIPLEVEL_MAX) );
+	check(Args.GetTabId() == ViewportTabId);
+
+	return SNew(SDockTab)
+		.Label(LOCTEXT("TextureViewportTitle", "Viewport"))
+		[
+			TextureViewport.ToSharedRef()
+		];
 }
 
 
-FReply FTextureEditorToolkit::OnMipLevelDown()
-{
-	if (SpecifiedMipLevel > MIPLEVEL_MIN)
-	{
-		--SpecifiedMipLevel;
-	}
-
-	return FReply::Handled();
-}
-
-
-FReply FTextureEditorToolkit::OnMipLevelUp()
-{
-	if (SpecifiedMipLevel < GetMaxMipLevel().Get(MIPLEVEL_MAX))
-	{
-		++SpecifiedMipLevel;
-	}
-
-	return FReply::Handled();
-}
-
-
-TOptional<int32> FTextureEditorToolkit::GetMaxMipLevel()const
-{
-	const UTexture2D* Texture2D = Cast<UTexture2D>(Texture);
-	const UTextureCube* TextureCube = Cast<UTextureCube>(Texture);
-	const UTextureRenderTargetCube* RTTextureCube = Cast<UTextureRenderTargetCube>(Texture);
-	const UTextureRenderTarget2D* RTTexture2D = Cast<UTextureRenderTarget2D>(Texture);
-	
-	if (Texture2D)
-	{
-		return Texture2D->GetNumMips() - 1;
-	}
-	
-	if (TextureCube)
-	{
-		return TextureCube->GetNumMips() - 1;
-	}
-	
-	if (RTTextureCube)
-	{
-		return RTTextureCube->GetNumMips() - 1;
-	}
-	
-	if (RTTexture2D)
-	{
-		return RTTexture2D->GetNumMips() - 1;
-	}
-
-	return MIPLEVEL_MAX;
-}
-
-
-bool FTextureEditorToolkit::GetUseSpecifiedMip() const
-{
-	if( GetMaxMipLevel().Get(MIPLEVEL_MAX) > 0 )
-	{
-		if( OnGetUseSpecifiedMipEnabled() )
-		{
-			return bUseSpecifiedMipLevel;
-		}
-
-		// by default this is on
-		return true; 
-	}
-
-	// disable the widgets if we have no mips
-	return false;
-}
-
-
-double FTextureEditorToolkit::GetZoom() const
-{
-	return Zoom;
-}
-
-
-void FTextureEditorToolkit::SetZoom( double ZoomValue )
-{
-	Zoom = FMath::Clamp(ZoomValue, MinZoom, MaxZoom);
-	SetFitToViewport(false);
-}
-
-
-void FTextureEditorToolkit::ZoomIn()
-{
-	SetZoom(Zoom + ZoomStep);
-}
-
-
-void FTextureEditorToolkit::ZoomOut()
-{
-	SetZoom(Zoom - ZoomStep);
-}
-
-
-bool FTextureEditorToolkit::GetFitToViewport() const
-{
-	const UTextureEditorSettings& Settings = *GetDefault<UTextureEditorSettings>();
-	return Settings.FitToViewport;
-}
-
-
-void FTextureEditorToolkit::SetFitToViewport(const bool bFitToViewport)
+void FTextureEditorToolkit::HandleTextureBorderActionExecute( )
 {
 	UTextureEditorSettings& Settings = *GetMutableDefault<UTextureEditorSettings>();
-	Settings.FitToViewport = bFitToViewport;
+	Settings.TextureBorderEnabled = !Settings.TextureBorderEnabled;
 	Settings.PostEditChange();
 }
 
 
-ESlateCheckBoxState::Type FTextureEditorToolkit::OnGetUseSpecifiedMip() const
+bool FTextureEditorToolkit::HandleTextureBorderActionIsChecked( ) const
 {
-	return GetUseSpecifiedMip() ? ESlateCheckBoxState::Checked : ESlateCheckBoxState::Unchecked;
+	const UTextureEditorSettings& Settings = *GetDefault<UTextureEditorSettings>();
+
+	return Settings.TextureBorderEnabled;
 }
 
 
-bool FTextureEditorToolkit::OnGetUseSpecifiedMipEnabled() const
-{
-	UTextureCube* TextureCube = Cast<UTextureCube>(Texture);
-
-	if( GetMaxMipLevel().Get(MIPLEVEL_MAX) <= 0 || TextureCube )
-	{
-		return false;
-	}
-
-	return true;
-}
-
-
-void FTextureEditorToolkit::OnUseSpecifiedMipChanged(ESlateCheckBoxState::Type InNewState)
-{
-	bUseSpecifiedMipLevel = InNewState == ESlateCheckBoxState::Checked;
-}
-
-
diff --git a/Engine/Source/Editor/TextureEditor/Private/TextureEditorToolkit.h b/Engine/Source/Editor/TextureEditor/Private/TextureEditorToolkit.h
--- a/Engine/Source/Editor/TextureEditor/Private/TextureEditorToolkit.h
+++ b/Engine/Source/Editor/TextureEditor/Private/TextureEditorToolkit.h
@@ -1,9 +1,5 @@
 // Copyright 1998-2014 Epic Games, Inc. All Rights Reserved.
 
-/*=============================================================================
-	TextureEditorToolkit.h: Declares the FTextureEditorToolkit class.
-=============================================================================*/
-
 #pragma once
 
 #include "Toolkits/AssetEditorToolkit.h"
@@ -29,156 +25,204 @@ public:
 	/**
 	 * Edits the specified Texture object.
 	 *
-	 * @param Mode -
-	 * @param InitToolkitHost -
-	 * @param ObjectToEdit -
+	 * @param Mode The tool kit mode.
+	 * @param InitToolkitHost 
+	 * @param ObjectToEdit The texture object to edit.
 	 */
-	void InitTextureEditor( const EToolkitMode::Type Mode, const TSharedPtr< class IToolkitHost >& InitToolkitHost, UObject* ObjectToEdit );
+	void InitTextureEditor( const EToolkitMode::Type Mode, const TSharedPtr<class IToolkitHost>& InitToolkitHost, UObject* ObjectToEdit );
 
 public:
 
-	// Begin FAssetEditorToolkit interface
+	// FAssetEditorToolkit interface
 
+	virtual FString GetDocumentationLink( ) const override;
 	virtual void RegisterTabSpawners( const TSharedRef<class FTabManager>& TabManager ) override;
 	virtual void UnregisterTabSpawners( const TSharedRef<class FTabManager>& TabManager ) override;
 
-	virtual FString GetDocumentationLink() const override
-	{
-		return FString(TEXT("Engine/Content/Types/Textures/Properties/Interface"));
-	}
-
-	// End FAssetEditorToolkit interface
-
 public:
 
-	// Begin ITextureEditorToolkit interface
+	// ITextureEditorToolkit interface
 
 	virtual void CalculateTextureDimensions( uint32& Width, uint32& Height ) const override;
-
 	virtual ESimpleElementBlendMode GetColourChannelBlendMode( ) const override;
-
-	virtual int32 GetMipLevel( ) const override
-	{
-		return GetUseSpecifiedMip() ? SpecifiedMipLevel : 0;
-	}
-
-	virtual UTexture* GetTexture( ) const override
-	{
-		return Texture;
-	}
-
+	virtual bool GetFitToViewport( ) const override;
+	virtual int32 GetMipLevel( ) const override;
+	virtual UTexture* GetTexture( ) const override;
+	virtual bool GetUseSpecifiedMip( ) const override;
+	virtual double GetZoom( ) const override;
 	virtual void PopulateQuickInfo( ) override;
-	
-	// End ITextureEditorToolkit interface
+	virtual void SetFitToViewport( const bool bFitToViewport ) override;
+	virtual void SetZoom( double ZoomValue ) override;
+	virtual void ZoomIn( ) override;
+	virtual void ZoomOut( ) override;
 
 public:
 
-	// Begin IToolkit interface
+	// IToolkit interface
 
 	virtual FText GetBaseToolkitName( ) const override;
+	virtual FName GetToolkitFName( ) const override;
+	virtual FLinearColor GetWorldCentricTabColorScale( ) const override;
+	virtual FString GetWorldCentricTabPrefix( ) const override;
 
-	virtual FName GetToolkitFName( ) const override
-	{
-		return FName("TextureEditor");
-	}
+public:
 
-	virtual FLinearColor GetWorldCentricTabColorScale( ) const override
-	{
-		return FLinearColor(0.3f, 0.2f, 0.5f, 0.5f);
-	}
+	// FGCObject interface
 
-	virtual bool GetUseSpecifiedMip( ) const override;
+	virtual void AddReferencedObjects( FReferenceCollector& Collector ) override;
+	
+protected:
 
-	virtual FString GetWorldCentricTabPrefix( ) const override;
+	// FEditorUndoClient interface
 
-	virtual double GetZoom( ) const override;
-	virtual void SetZoom( double ZoomValue ) override;
-	virtual void ZoomIn( ) override;
-	virtual void ZoomOut( ) override;
-	virtual bool GetFitToViewport( ) const override;
-	virtual void SetFitToViewport( const bool bFitToViewport ) override;
+	virtual void PostUndo( bool bSuccess ) override;
+	virtual void PostRedo( bool bSuccess ) override;
 
-	// End IToolkit interface
+protected:
 
-public:
+	/**
+     * Binds the UI commands to delegates.
+	 */
+	void BindCommands( );
 
-	// Begin FGCObject interface
+	/**
+	 * Creates the texture properties details widget.
+	 *
+	 * @return The widget.
+	 */
+	TSharedRef<SWidget> BuildTexturePropertiesWidget( );
 
-	virtual void AddReferencedObjects( FReferenceCollector& Collector ) override;
+	/**
+	 * Calculates the effective in-game resolution of the texture.
+	 *
+	 * @param LODBias 
+	 * @param EffectiveTextureWidth 
+	 * @param EffectiveTextureHeight 
+	 */
+	void CalculateEffectiveTextureDimensions( int32 LODBias, uint32& EffectiveTextureWidth, uint32& EffectiveTextureHeight );
 
-	// End FGCObject interface
-	
-protected:
+	/**
+	 * Creates all internal widgets for the tabs to point at.
+	 */
+	void CreateInternalWidgets( );
 
-	// Begin FEditorUndoClient interface
+	/**
+	 * Builds the toolbar widget for the Texture editor.
+	 */
+	void ExtendToolBar( );
 
-	virtual void PostUndo(bool bSuccess) override { }
+	/**
+	 * Gets the highest mip map level that this texture supports.
+	 *
+	 * @return Mip map level.
+	 */
+	TOptional<int32> GetMaxMipLevel( ) const;
 
-	virtual void PostRedo(bool bSuccess) override
-	{
-		PostUndo(bSuccess);
-	}
-	
-	// End FEditorUndoClient interface
+	/**
+	 * Checks whether the texture being edited is a cube map texture.
+	 */
+	bool IsCubeTexture( ) const;
 
 private:
 
-	TSharedRef<SDockTab> SpawnTab_Viewport(const FSpawnTabArgs& Args);
-	TSharedRef<SDockTab> SpawnTab_Properties(const FSpawnTabArgs& Args);
+	// Callback for toggling the Alpha channel action.
+	void HandleAlphaChannelActionExecute( );
 
-	/** Creates the texture properties details widget */
-	TSharedRef<SWidget> BuildTexturePropertiesWidget();
+	// Callback for getting the checked state of the Alpha channel action.
+	bool HandleAlphaChannelActionIsChecked( ) const;
 
-	/** Creates all internal widgets for the tabs to point at */
-	void CreateInternalWidgets();
+	// Callback for getting the enabled state of the Alpha channel action.
+	bool HandleAlphaChannelActionCanExecute( ) const;
 
-	/** Builds the toolbar widget for the Texture editor */
-	void ExtendToolBar();
-	
-	/**	Binds our UI commands to delegates */
-	void BindCommands();
+	// Callback for toggling the Blue channel action.
+	void HandleBlueChannelActionExecute( );
+
+	// Callback for getting the checked state of the Blue channel action.
+	bool HandleBlueChannelActionIsChecked( ) const;
+
+	// Callback for toggling the Checkered Background action.
+	void HandleCheckeredBackgroundActionExecute( ETextureEditorBackgrounds Background );
+
+	// Callback for getting the checked state of the Checkered Background action.
+	bool HandleCheckeredBackgroundActionIsChecked( ETextureEditorBackgrounds Background );
+
+	// Callback for toggling the Compress Now action.
+	void HandleCompressNowActionExecute( );
+
+	// Callback for getting the checked state of the Compress Now action.
+	bool HandleCompressNowActionCanExecute( ) const;
+
+	// Callback for toggling the Fit To Viewport action.
+	void HandleFitToViewportActionExecute( );
+
+	// Callback for getting the checked state of the Fit To Viewport action.
+	bool HandleFitToViewportActionIsChecked( ) const;
+
+	// Callback for toggling the Green channel action.
+	void HandleGreenChannelActionExecute( );
 
-	/**	Finds the effective in-game resolution of the texture */
-	void CalculateEffectiveTextureDimensions(int32 LODBias, uint32& EffectiveTextureWidth, uint32& EffectiveTextureHeight);
+	// Callback for getting the checked state of the Green Channel action.
+	bool HandleGreenChannelActionIsChecked( ) const;
+
+	// Callback for changing the checked state of the MipMap check box.
+	void HandleMipLevelCheckBoxCheckedStateChanged( ESlateCheckBoxState::Type InNewState );
+
+	// Callback for getting the checked state of the MipMap check box.
+	ESlateCheckBoxState::Type HandleMipLevelCheckBoxIsChecked( ) const;
+
+	// Callback for determining whether the MipMap check box is enabled.
+	bool HandleMipLevelCheckBoxIsEnabled( ) const;
+
+	// Callback for changing the value of the mip map level entry box.
+	void HandleMipLevelEntryBoxChanged( int32 NewMipLevel );
+
+	// Callback for getting the value of the mip map level entry box.
+	TOptional<int32> HandleMipLevelEntryBoxValue( ) const;
+
+	// Callback for clicking the MipMinus button.
+	FReply HandleMipMapMinusButtonClicked( );
+
+	// Callback for clicking the MipPlus button.
+	FReply HandleMipMapPlusButtonClicked( );
+
+	// Callback for toggling the Red channel action.
+	void HandleRedChannelActionExecute( );
+
+	// Callback for getting the checked state of the Red Channel action.
+	bool HandleRedChannelActionIsChecked( ) const;
+
+	// Callback for determining whether the Reimport action can execute.
+	bool HandleReimportActionCanExecute( ) const;
+
+	// Callback for executing the Reimport action.
+	void HandleReimportActionExecute( );
+
+	// Callback that is executed after the reimport manager reimported an asset.
+	void HandleReimportManagerPostReimport( UObject* InObject, bool bSuccess );
 	
-	/** Toolbar command methods */
-	void OnRedChannel();
-	bool IsRedChannelChecked() const;
-	void OnGreenChannel();
-	bool IsGreenChannelChecked() const;
-	void OnBlueChannel();
-	bool IsBlueChannelChecked() const;
-	void OnAlphaChannel();
-	bool OnAlphaChannelCanExecute() const;
-	bool IsAlphaChannelChecked() const;
-	void OnSaturation();
-	bool IsSaturationChecked() const;
-	void OnFitToViewport();
-	bool IsFitToViewportChecked() const;
-	void HandleBackgroundActionExecute( ETextureEditorBackgrounds Background );
-	bool HandleBackgroundActionIsChecked( ETextureEditorBackgrounds Background );
-	void OnShowBorder();
-	bool IsShowBorderChecked() const;
-	void OnPreReimport(UObject* InObject);
-	void OnPostReimport(UObject* InObject, bool bSuccess);
-	void OnCompressNow();
-	bool OnCompressNowEnabled() const;
-	void OnReimport();
-	bool OnReimportEnabled() const;
-	void HandleSettingsActionExecute();
-	bool IsCubeTexture() const;
-
-	/** Methods for the mip level UI control */
-	TOptional<int32> GetWidgetMipLevel() const;
-	void OnMipLevelChanged(int32 NewMipLevel);
-	FReply OnMipLevelDown();
-	FReply OnMipLevelUp();
-
-	TOptional<int32> GetMaxMipLevel()const;
-
-	ESlateCheckBoxState::Type OnGetUseSpecifiedMip()const;
-	bool OnGetUseSpecifiedMipEnabled()const;
-	void OnUseSpecifiedMipChanged(ESlateCheckBoxState::Type InNewState);
+	// Callback that is executed before the reimport manager reimported an asset.
+	void HandleReimportManagerPreReimport( UObject* InObject );
+
+	// Callback for toggling the Saturation channel action.
+	void HandleSaturationChannelActionExecute( );
+
+	// Callback for getting the checked state of the Saturation channel action.
+	bool HandleSaturationChannelActionIsChecked( ) const;
+
+	// Callback for determining whether the Settings action can execute.
+	void HandleSettingsActionExecute( );
+
+	// Callback for spawning the Properties tab.
+	TSharedRef<SDockTab> HandleTabSpawnerSpawnProperties( const FSpawnTabArgs& Args );
+
+	// Callback for spawning the Viewport tab.
+	TSharedRef<SDockTab> HandleTabSpawnerSpawnViewport( const FSpawnTabArgs& Args );
+
+	// Callback for toggling the Texture Border action.
+	void HandleTextureBorderActionExecute( );
+
+	// Callback for getting the checked state of the Texture Border action.
+	bool HandleTextureBorderActionIsChecked( ) const;
 
 private:
 
@@ -186,7 +230,7 @@ private:
 	UTexture* Texture;
 
 	/** List of open tool panels; used to ensure only one exists at any one time */
-	TMap< FName, TWeakPtr<SDockableTab> > SpawnedToolPanels;
+	TMap<FName, TWeakPtr<SDockableTab>> SpawnedToolPanels;
 
 	/** Viewport */
 	TSharedPtr<STextureEditorViewport> TextureViewport;
@@ -237,7 +281,9 @@ private:
 
 private:
 
-	/**	The tab ids for all the tabs used */
+	// The name of the Viewport tab.
 	static const FName ViewportTabId;
+
+	// The name of the Properties tab.
 	static const FName PropertiesTabId;
 };
diff --git a/Engine/Source/Editor/TextureEditor/Private/Widgets/STextureEditorViewportToolbar.cpp b/Engine/Source/Editor/TextureEditor/Private/Widgets/STextureEditorViewportToolbar.cpp
--- a/Engine/Source/Editor/TextureEditor/Private/Widgets/STextureEditorViewportToolbar.cpp
+++ b/Engine/Source/Editor/TextureEditor/Private/Widgets/STextureEditorViewportToolbar.cpp
@@ -1,11 +1,8 @@
 // Copyright 1998-2014 Epic Games, Inc. All Rights Reserved.
 
-/*=============================================================================
-	STextureEditorViewportToolbar.cpp: Implements the STextureEditorViewportToolbar class.
-=============================================================================*/
-
 #include "TextureEditorPrivatePCH.h"
 
+
 #define LOCTEXT_NAMESPACE "STextureEditorViewportToolbar"
 
 
@@ -24,6 +21,7 @@ void STextureEditorViewportToolbar::Construct( const FArguments& InArgs, const T
 			.AutoWidth()
 			[
 				SAssignNew(ViewOptionsMenuAnchor, SMenuAnchor)
+					.OnGetMenuContent(this, &STextureEditorViewportToolbar::GenerateViewOptionsMenu)
 					.Placement(MenuPlacement_ComboBox)
 					[
 						SNew(SButton)
@@ -51,7 +49,6 @@ void STextureEditorViewportToolbar::Construct( const FArguments& InArgs, const T
 									]
 							]
 					]
-					.OnGetMenuContent(this, &STextureEditorViewportToolbar::GenerateViewOptionsMenu)
 			]
 	];
 }
@@ -71,14 +68,13 @@ TSharedRef<SWidget> STextureEditorViewportToolbar::GenerateViewOptionsMenu( ) co
 
 FReply STextureEditorViewportToolbar::HandleViewOptionsMenuButtonClicked( )
 {
-	// If the menu button is clicked toggle the state of the menu anchor which will open or close the menu
 	if (ViewOptionsMenuAnchor->ShouldOpenDueToClick())
 	{
-		ViewOptionsMenuAnchor->SetIsOpen( true );
+		ViewOptionsMenuAnchor->SetIsOpen(true);
 	}
 	else
 	{
-		ViewOptionsMenuAnchor->SetIsOpen( false );
+		ViewOptionsMenuAnchor->SetIsOpen(false);
 	}
 
 	return FReply::Handled();
diff --git a/Engine/Source/Editor/TextureEditor/Private/Widgets/STextureEditorViewportToolbar.h b/Engine/Source/Editor/TextureEditor/Private/Widgets/STextureEditorViewportToolbar.h
--- a/Engine/Source/Editor/TextureEditor/Private/Widgets/STextureEditorViewportToolbar.h
+++ b/Engine/Source/Editor/TextureEditor/Private/Widgets/STextureEditorViewportToolbar.h
@@ -1,9 +1,5 @@
 // Copyright 1998-2014 Epic Games, Inc. All Rights Reserved.
 
-/*=============================================================================
-	STextureEditorViewportToolbar.h: Declares the STextureEditorViewportToolbar class.
-=============================================================================*/
-
 #pragma once
 
 
@@ -23,11 +19,10 @@ public:
 	/**
 	 * Constructs the widget.
 	 *
-	 * @param InArgs - The construction arguments.
+	 * @param InArgs The construction arguments.
 	 */
 	void Construct( const FArguments& InArgs, const TSharedRef<FUICommandList>& InToolkitCommands );
 
-
 protected:
 
 	/**
