Editor: Default texture editor to showing the alpha channel in some texture formats
Editor: Correct actual texture size display in texture editor when using a power-of-2 padding option

diff --git a/Engine/Source/Editor/TextureEditor/Private/TextureEditorToolkit.cpp b/Engine/Source/Editor/TextureEditor/Private/TextureEditorToolkit.cpp
--- a/Engine/Source/Editor/TextureEditor/Private/TextureEditorToolkit.cpp
+++ b/Engine/Source/Editor/TextureEditor/Private/TextureEditorToolkit.cpp
@@ -93,6 +93,21 @@ void FTextureEditorToolkit::InitTextureEditor( const EToolkitMode::Type Mode, co
 	bIsGreenChannel = true;
 	bIsBlueChannel = true;
 	bIsAlphaChannel = false;
+
+	switch (Texture->CompressionSettings)
+	{
+	default:
+		bIsAlphaChannel = !Texture->CompressionNoAlpha;
+		break;
+	case TC_Normalmap:
+	case TC_Grayscale:
+	case TC_Displacementmap:
+	case TC_VectorDisplacementmap:
+	case TC_DistanceFieldFont:
+		bIsAlphaChannel = false;
+		break;
+	}
+
 	bIsSaturation = false;
 
 	PreviewEffectiveTextureWidth = 0;
@@ -331,6 +346,9 @@ void FTextureEditorToolkit::PopulateQuickInfo( )
 		ImportedHeight = Texture->GetSurfaceHeight();
 	}
 
+	const uint32 ActualWidth = (Texture->GetSurfaceWidth() > 0) ? (uint32)Texture->GetSurfaceWidth() : ImportedWidth;
+	const uint32 ActualHeight = (Texture->GetSurfaceHeight() > 0) ? (uint32)Texture->GetSurfaceHeight() : ImportedHeight;
+
 	// In game max bias and dimensions
 	uint32 MaxInGameWidth, MaxInGameHeight;
 	int32 MipLevel = GSystemSettings.TextureLODSettings.CalculateLODBias(Texture);
@@ -349,7 +367,7 @@ void FTextureEditorToolkit::PopulateQuickInfo( )
 	Options.UseGrouping = false;
 
 	ImportedText->SetText(FText::Format( NSLOCTEXT("TextureEditor", "QuickInfo_Imported", "Imported: {0}x{1}"), FText::AsNumber(ImportedWidth, &Options), FText::AsNumber(ImportedHeight, &Options)));
-	CurrentText->SetText(FText::Format( NSLOCTEXT("TextureEditor", "QuickInfo_Displayed", "Displayed: {0}x{1}"), FText::AsNumber(FMath::Max((uint32)1, ImportedWidth >> MipLevel), &Options ), FText::AsNumber(FMath::Max((uint32)1, ImportedHeight>> MipLevel), &Options)));
+	CurrentText->SetText(FText::Format( NSLOCTEXT("TextureEditor", "QuickInfo_Displayed", "Displayed: {0}x{1}"), FText::AsNumber(FMath::Max((uint32)1, ActualWidth >> MipLevel), &Options ), FText::AsNumber(FMath::Max((uint32)1, ActualHeight >> MipLevel), &Options)));
 	MaxInGameText->SetText(FText::Format( NSLOCTEXT("TextureEditor", "QuickInfo_MaxInGame", "Max In-Game: {0}x{1}"), FText::AsNumber(MaxInGameWidth, &Options), FText::AsNumber(MaxInGameHeight, &Options)));
 	MethodText->SetText(FText::Format( NSLOCTEXT("TextureEditor", "QuickInfo_Method", "Method: {0}"), Texture->NeverStream ? NSLOCTEXT("TextureEditor", "QuickInfo_MethodNotStreamed", "Not Streamed") : NSLOCTEXT("TextureEditor", "QuickInfo_MethodStreamed", "Streamed")));
 	LODBiasText->SetText(FText::Format(NSLOCTEXT("TextureEditor", "QuickInfo_LODBias", "Combined LOD Bias: {0}"), FText::AsNumber(Texture->GetCachedLODBias())));
