fix for UE-5769 Crash when importing png to cube render target (fix 2)

diff --git a/Engine/Source/Editor/TextureEditor/Private/TextureEditorToolkit.cpp b/Engine/Source/Editor/TextureEditor/Private/TextureEditorToolkit.cpp
--- a/Engine/Source/Editor/TextureEditor/Private/TextureEditorToolkit.cpp
+++ b/Engine/Source/Editor/TextureEditor/Private/TextureEditorToolkit.cpp
@@ -6,6 +6,7 @@
 #include "SDockTab.h"
 #include "SNumericEntryBox.h"
 #include "Engine/TextureCube.h"
+#include "Engine/TextureRenderTarget.h"
 #include "Engine/TextureRenderTargetCube.h"
 #include "Engine/TextureRenderTarget2D.h"
 #include "Engine/Texture2D.h"
@@ -1001,7 +1002,7 @@ bool FTextureEditorToolkit::HandleRedChannelActionIsChecked( ) const
 
 bool FTextureEditorToolkit::HandleReimportActionCanExecute( ) const
 {
-	if (Texture->IsA<ULightMapTexture2D>() || Texture->IsA<UShadowMapTexture2D>() || Texture->IsA<UTexture2DDynamic>())
+	if (Texture->IsA<ULightMapTexture2D>() || Texture->IsA<UShadowMapTexture2D>() || Texture->IsA<UTexture2DDynamic>() || Texture->IsA<UTextureRenderTarget>())
 	{
 		return false;
 	}
