Changed default Texture Editor background to be Checkered.  Renamed Saturation view option to Desaturation.

diff --git a/Engine/Source/Editor/TextureEditor/Private/Menus/TextureEditorViewOptionsMenu.h b/Engine/Source/Editor/TextureEditor/Private/Menus/TextureEditorViewOptionsMenu.h
--- a/Engine/Source/Editor/TextureEditor/Private/Menus/TextureEditorViewOptionsMenu.h
+++ b/Engine/Source/Editor/TextureEditor/Private/Menus/TextureEditorViewOptionsMenu.h
@@ -27,7 +27,7 @@ public:
 			MenuBuilder.AddMenuEntry(FTextureEditorCommands::Get().GreenChannel);
 			MenuBuilder.AddMenuEntry(FTextureEditorCommands::Get().BlueChannel);
 			MenuBuilder.AddMenuEntry(FTextureEditorCommands::Get().AlphaChannel);
-			MenuBuilder.AddMenuEntry(FTextureEditorCommands::Get().Saturation);
+			MenuBuilder.AddMenuEntry(FTextureEditorCommands::Get().Desaturation);
 		}
 		MenuBuilder.EndSection();
 
diff --git a/Engine/Source/Editor/TextureEditor/Private/Models/TextureEditorCommands.cpp b/Engine/Source/Editor/TextureEditor/Private/Models/TextureEditorCommands.cpp
--- a/Engine/Source/Editor/TextureEditor/Private/Models/TextureEditorCommands.cpp
+++ b/Engine/Source/Editor/TextureEditor/Private/Models/TextureEditorCommands.cpp
@@ -9,7 +9,7 @@ void FTextureEditorCommands::RegisterCommands()
 	UI_COMMAND(GreenChannel, "Green", "Toggles the green channel", EUserInterfaceActionType::ToggleButton, FInputGesture());
 	UI_COMMAND(BlueChannel, "Blue", "Toggles the blue channel", EUserInterfaceActionType::ToggleButton, FInputGesture());
 	UI_COMMAND(AlphaChannel, "Alpha", "Toggles the alpha channel", EUserInterfaceActionType::ToggleButton, FInputGesture());
-	UI_COMMAND(Saturation, "Saturation", "Toggles color saturation", EUserInterfaceActionType::ToggleButton, FInputGesture());
+	UI_COMMAND(Desaturation, "Desaturation", "Toggles color desaturation", EUserInterfaceActionType::ToggleButton, FInputGesture());
 	
 	UI_COMMAND(CheckeredBackground, "Checkered", "Checkered background pattern behind the texture", EUserInterfaceActionType::RadioButton, FInputGesture());
 	UI_COMMAND(CheckeredBackgroundFill, "Checkered (Fill)", "Checkered background pattern behind the entire viewport", EUserInterfaceActionType::RadioButton, FInputGesture());
diff --git a/Engine/Source/Editor/TextureEditor/Private/Models/TextureEditorCommands.h b/Engine/Source/Editor/TextureEditor/Private/Models/TextureEditorCommands.h
--- a/Engine/Source/Editor/TextureEditor/Private/Models/TextureEditorCommands.h
+++ b/Engine/Source/Editor/TextureEditor/Private/Models/TextureEditorCommands.h
@@ -39,7 +39,7 @@ public:
 	TSharedPtr<FUICommandInfo> AlphaChannel;
 
 	/** Toggles color saturation */
-	TSharedPtr<FUICommandInfo> Saturation;
+	TSharedPtr<FUICommandInfo> Desaturation;
 
 	/** If enabled, the texture will be scaled to fit the viewport */
 	TSharedPtr<FUICommandInfo> FitToViewport;
diff --git a/Engine/Source/Editor/TextureEditor/Private/TextureEditorClasses.cpp b/Engine/Source/Editor/TextureEditor/Private/TextureEditorClasses.cpp
--- a/Engine/Source/Editor/TextureEditor/Private/TextureEditorClasses.cpp
+++ b/Engine/Source/Editor/TextureEditor/Private/TextureEditorClasses.cpp
@@ -5,7 +5,7 @@
 
 UTextureEditorSettings::UTextureEditorSettings( const FObjectInitializer& ObjectInitializer )
 	: Super(ObjectInitializer)
-	, Background(TextureEditorBackground_SolidColor)
+	, Background(TextureEditorBackground_Checkered)
 	, BackgroundColor(FColor::Black)
 	, CheckerColorOne(FColor(128, 128, 128))
 	, CheckerColorTwo(FColor(64, 64, 64))
diff --git a/Engine/Source/Editor/TextureEditor/Private/TextureEditorToolkit.cpp b/Engine/Source/Editor/TextureEditor/Private/TextureEditorToolkit.cpp
--- a/Engine/Source/Editor/TextureEditor/Private/TextureEditorToolkit.cpp
+++ b/Engine/Source/Editor/TextureEditor/Private/TextureEditorToolkit.cpp
@@ -109,7 +109,7 @@ void FTextureEditorToolkit::InitTextureEditor( const EToolkitMode::Type Mode, co
 		break;
 	}
 
-	bIsSaturation = false;
+	bIsDesaturation = false;
 
 	PreviewEffectiveTextureWidth = 0;
 	PreviewEffectiveTextureHeight = 0;
@@ -280,8 +280,8 @@ ESimpleElementBlendMode FTextureEditorToolkit::GetColourChannelBlendMode( ) cons
 	
 	// If we only have one color channel active, enable color desaturation by default
 	const int32 NumColorChannelsActive = (bIsRedChannel ? 1 : 0) + (bIsGreenChannel ? 1 : 0) + (bIsBlueChannel ? 1 : 0);
-	const bool bIsSaturationLocal = bIsSaturation ? true : (NumColorChannelsActive==1);
-	Result += bIsSaturationLocal ? (1 << 4) : 0;
+	const bool bIsDesaturationLocal = bIsDesaturation ? true : (NumColorChannelsActive==1);
+	Result += bIsDesaturationLocal ? (1 << 4) : 0;
 
 	return (ESimpleElementBlendMode)Result;
 }
@@ -505,10 +505,10 @@ void FTextureEditorToolkit::BindCommands( )
 		FIsActionChecked::CreateSP(this, &FTextureEditorToolkit::HandleAlphaChannelActionIsChecked));
 
 	ToolkitCommands->MapAction(
-		Commands.Saturation,
-		FExecuteAction::CreateSP(this, &FTextureEditorToolkit::HandleSaturationChannelActionExecute),
+		Commands.Desaturation,
+		FExecuteAction::CreateSP(this, &FTextureEditorToolkit::HandleDesaturationChannelActionExecute),
 		FCanExecuteAction(),
-		FIsActionChecked::CreateSP(this, &FTextureEditorToolkit::HandleSaturationChannelActionIsChecked));
+		FIsActionChecked::CreateSP(this, &FTextureEditorToolkit::HandleDesaturationChannelActionIsChecked));
 
 	ToolkitCommands->MapAction(
 		Commands.FitToViewport,
@@ -1067,15 +1067,15 @@ void FTextureEditorToolkit::HandleReimportManagerPreReimport( UObject* InObject
 }
 
 
-void FTextureEditorToolkit::HandleSaturationChannelActionExecute( )
+void FTextureEditorToolkit::HandleDesaturationChannelActionExecute( )
 {
-	bIsSaturation = !bIsSaturation;
+	bIsDesaturation = !bIsDesaturation;
 }
 
 
-bool FTextureEditorToolkit::HandleSaturationChannelActionIsChecked( ) const
+bool FTextureEditorToolkit::HandleDesaturationChannelActionIsChecked( ) const
 {
-	return bIsSaturation;
+	return bIsDesaturation;
 }
 
 
diff --git a/Engine/Source/Editor/TextureEditor/Private/TextureEditorToolkit.h b/Engine/Source/Editor/TextureEditor/Private/TextureEditorToolkit.h
--- a/Engine/Source/Editor/TextureEditor/Private/TextureEditorToolkit.h
+++ b/Engine/Source/Editor/TextureEditor/Private/TextureEditorToolkit.h
@@ -204,11 +204,11 @@ private:
 	// Callback that is executed before the reimport manager reimported an asset.
 	void HandleReimportManagerPreReimport( UObject* InObject );
 
-	// Callback for toggling the Saturation channel action.
-	void HandleSaturationChannelActionExecute( );
+	// Callback for toggling the Desaturation channel action.
+	void HandleDesaturationChannelActionExecute( );
 
-	// Callback for getting the checked state of the Saturation channel action.
-	bool HandleSaturationChannelActionIsChecked( ) const;
+	// Callback for getting the checked state of the Desaturation channel action.
+	bool HandleDesaturationChannelActionIsChecked( ) const;
 
 	// Callback for determining whether the Settings action can execute.
 	void HandleSettingsActionExecute( );
@@ -263,7 +263,7 @@ private:
 	bool bIsAlphaChannel;
 
 	/** If true, desaturates the texture */
-	bool bIsSaturation;
+	bool bIsDesaturation;
 
 	/** The maximum width/height at which the texture will render in the preview window */
 	uint32 PreviewEffectiveTextureWidth;
