Fixed ResourceSize tag for texture assets so it updates as new mip levels are streamed in.  Added extra data to the details panel in the Texture Editor.

diff --git a/Engine/Source/Editor/TextureEditor/Private/TextureEditorToolkit.cpp b/Engine/Source/Editor/TextureEditor/Private/TextureEditorToolkit.cpp
--- a/Engine/Source/Editor/TextureEditor/Private/TextureEditorToolkit.cpp
+++ b/Engine/Source/Editor/TextureEditor/Private/TextureEditorToolkit.cpp
@@ -360,6 +360,13 @@ void FTextureEditorToolkit::PopulateQuickInfo( )
 	MipLevel = FMath::Max(GetMipLevel(), Texture->GetCachedLODBias());
 	CalculateEffectiveTextureDimensions(MipLevel, PreviewEffectiveTextureWidth, PreviewEffectiveTextureHeight);
 
+	// Texture asset size
+	const SIZE_T Size = Texture->GetResourceSize(EResourceSizeMode::Exclusive);
+
+	FNumberFormattingOptions SizeOptions;
+	SizeOptions.UseGrouping = false;
+	SizeOptions.MaximumFractionalDigits = 0;
+
 	// Cubes are previewed as unwrapped 2D textures.
 	// These have 2x the width of a cube face.
 	PreviewEffectiveTextureWidth *= IsCubeTexture() ? 2 : 1;
@@ -370,9 +377,10 @@ void FTextureEditorToolkit::PopulateQuickInfo( )
 	ImportedText->SetText(FText::Format( NSLOCTEXT("TextureEditor", "QuickInfo_Imported", "Imported: {0}x{1}"), FText::AsNumber(ImportedWidth, &Options), FText::AsNumber(ImportedHeight, &Options)));
 	CurrentText->SetText(FText::Format( NSLOCTEXT("TextureEditor", "QuickInfo_Displayed", "Displayed: {0}x{1}"), FText::AsNumber(FMath::Max((uint32)1, ActualWidth >> MipLevel), &Options ), FText::AsNumber(FMath::Max((uint32)1, ActualHeight >> MipLevel), &Options)));
 	MaxInGameText->SetText(FText::Format( NSLOCTEXT("TextureEditor", "QuickInfo_MaxInGame", "Max In-Game: {0}x{1}"), FText::AsNumber(MaxInGameWidth, &Options), FText::AsNumber(MaxInGameHeight, &Options)));
-	MethodText->SetText(FText::Format( NSLOCTEXT("TextureEditor", "QuickInfo_Method", "Method: {0}"), Texture->NeverStream ? NSLOCTEXT("TextureEditor", "QuickInfo_MethodNotStreamed", "Not Streamed") : NSLOCTEXT("TextureEditor", "QuickInfo_MethodStreamed", "Streamed")));
+	SizeText->SetText(FText::Format(NSLOCTEXT("TextureEditor", "QuickInfo_ResourceSize", "Resource Size: {0}kb"), FText::AsNumber((Size + 512) / 1024, &SizeOptions)));
+	MethodText->SetText(FText::Format(NSLOCTEXT("TextureEditor", "QuickInfo_Method", "Method: {0}"), Texture->NeverStream ? NSLOCTEXT("TextureEditor", "QuickInfo_MethodNotStreamed", "Not Streamed") : NSLOCTEXT("TextureEditor", "QuickInfo_MethodStreamed", "Streamed")));
 	LODBiasText->SetText(FText::Format(NSLOCTEXT("TextureEditor", "QuickInfo_LODBias", "Combined LOD Bias: {0}"), FText::AsNumber(Texture->GetCachedLODBias())));
-	
+
 	int32 TextureFormatIndex = PF_MAX;
 	
 	if (Texture2D)
@@ -392,6 +400,26 @@ void FTextureEditorToolkit::PopulateQuickInfo( )
 	{
 		FormatText->SetText(FText::Format(NSLOCTEXT("TextureEditor", "QuickInfo_Format", "Format: {0}"), FText::FromString(GPixelFormats[TextureFormatIndex].Name)));
 	}
+
+	int32 NumMips = 1;
+	if (Texture2D)
+	{
+		NumMips = Texture2D->GetNumMips();
+	}
+	else if (TextureCube)
+	{
+		NumMips = TextureCube->GetNumMips();
+	}
+
+	NumMipsText->SetText(FText::Format(NSLOCTEXT("TextureEditor", "QuickInfo_NumMips", "Number of Mips: {0}"), FText::AsNumber(NumMips)));
+
+	if (Texture2D)
+	{
+		HasAlphaChannelText->SetText(FText::Format(NSLOCTEXT("TextureEditor", "QuickInfo_HasAlphaChannel", "Has Alpha Channel: {0}"),
+			Texture2D->HasAlphaChannel() ? NSLOCTEXT("TextureEditor", "True", "True") : NSLOCTEXT("TextureEditor", "False", "False")));
+	}
+
+	HasAlphaChannelText->SetVisibility(Texture2D ? EVisibility::Visible : EVisibility::Collapsed);
 }
 
 
@@ -583,85 +611,109 @@ void FTextureEditorToolkit::CreateInternalWidgets( )
 	TextureProperties = SNew(SVerticalBox)
 
 	+ SVerticalBox::Slot()
-		.AutoHeight()
-		.Padding(2.0f)
+	.AutoHeight()
+	.Padding(2.0f)
+	[
+		SNew(SBorder)
 		[
-			SNew(SBorder)
+			SNew(SHorizontalBox)
+
+			+ SHorizontalBox::Slot()
+			.FillWidth(0.5f)
+			[
+				SNew(SVerticalBox)
+
+				+ SVerticalBox::Slot()
+				.AutoHeight()
+				.VAlign(VAlign_Center)
+				.Padding(4.0f)
 				[
-					SNew(SHorizontalBox)
+					SAssignNew(ImportedText, STextBlock)
+				]
 
-					+ SHorizontalBox::Slot()
-						.FillWidth(0.5f)
-						[
-							SNew(SVerticalBox)
+				+ SVerticalBox::Slot()
+				.AutoHeight()
+				.VAlign(VAlign_Center)
+				.Padding(4.0f)
+				[
+					SAssignNew(CurrentText, STextBlock)
+				]
 
-							+ SVerticalBox::Slot()
-								.AutoHeight()
-								.VAlign(VAlign_Center)
-								.Padding(4.0f)
-								[
-									SAssignNew(ImportedText, STextBlock)
-								]
+				+ SVerticalBox::Slot()
+				.AutoHeight()
+				.VAlign(VAlign_Center)
+				.Padding(4.0f)
+				[
+					SAssignNew(MaxInGameText, STextBlock)
+				]
 
-							+ SVerticalBox::Slot()
-								.AutoHeight()
-								.VAlign(VAlign_Center)
-								.Padding(4.0f)
-								[
-									SAssignNew(CurrentText, STextBlock)
-								]
+				+ SVerticalBox::Slot()
+				.AutoHeight()
+				.VAlign(VAlign_Center)
+				.Padding(4.0f)
+				[
+					SAssignNew(SizeText, STextBlock)
+				]
 
-							+ SVerticalBox::Slot()
-								.AutoHeight()
-								.VAlign(VAlign_Center)
-								.Padding(4.0f)
-								[
-									SAssignNew(MaxInGameText, STextBlock)
-								]
-						]
+				+ SVerticalBox::Slot()
+				.AutoHeight()
+				.VAlign(VAlign_Center)
+				.Padding(4.0f)
+				[
+					SAssignNew(HasAlphaChannelText, STextBlock)
+				]
+			]
 
-					+ SHorizontalBox::Slot()
-						.FillWidth(0.5f)
-						[
-							SNew(SVerticalBox)
+			+ SHorizontalBox::Slot()
+			.FillWidth(0.5f)
+			[
+				SNew(SVerticalBox)
 
-							+ SVerticalBox::Slot()
-								.AutoHeight()
-								.VAlign(VAlign_Center)
-								.Padding(4.0f)
-								[
-									SAssignNew(MethodText, STextBlock)
-								]
+				+ SVerticalBox::Slot()
+				.AutoHeight()
+				.VAlign(VAlign_Center)
+				.Padding(4.0f)
+				[
+					SAssignNew(MethodText, STextBlock)
+				]
 
-							+ SVerticalBox::Slot()
-								.AutoHeight()
-								.VAlign(VAlign_Center)
-								.Padding(4.0f)
-								[
-									SAssignNew(FormatText, STextBlock)
-								]
+				+ SVerticalBox::Slot()
+				.AutoHeight()
+				.VAlign(VAlign_Center)
+				.Padding(4.0f)
+				[
+					SAssignNew(FormatText, STextBlock)
+				]
 
-							+ SVerticalBox::Slot()
-								.AutoHeight()
-								.VAlign(VAlign_Center)
-								.Padding(4.0f)
-								[
-									SAssignNew(LODBiasText, STextBlock)
-								]
-						]
+				+ SVerticalBox::Slot()
+				.AutoHeight()
+				.VAlign(VAlign_Center)
+				.Padding(4.0f)
+				[
+					SAssignNew(LODBiasText, STextBlock)
 				]
-		]
 
-	+ SVerticalBox::Slot()
-		.FillHeight(1.0f)
-		.Padding(2.0f)
-		[
-			SNew(SBorder)
+				+ SVerticalBox::Slot()
+				.AutoHeight()
+				.VAlign(VAlign_Center)
 				.Padding(4.0f)
 				[
-					BuildTexturePropertiesWidget()
+					SAssignNew(NumMipsText, STextBlock)
 				]
-		];
+			]
+		]
+	]
+
+	+ SVerticalBox::Slot()
+	.FillHeight(1.0f)
+	.Padding(2.0f)
+	[
+		SNew(SBorder)
+		.Padding(4.0f)
+		[
+			BuildTexturePropertiesWidget()
+		]
+	];
 }
 
 
diff --git a/Engine/Source/Editor/TextureEditor/Private/TextureEditorToolkit.h b/Engine/Source/Editor/TextureEditor/Private/TextureEditorToolkit.h
--- a/Engine/Source/Editor/TextureEditor/Private/TextureEditorToolkit.h
+++ b/Engine/Source/Editor/TextureEditor/Private/TextureEditorToolkit.h
@@ -246,9 +246,12 @@ private:
 	TSharedPtr<STextBlock> ImportedText;
 	TSharedPtr<STextBlock> CurrentText;
 	TSharedPtr<STextBlock> MaxInGameText;
+	TSharedPtr<STextBlock> SizeText;
 	TSharedPtr<STextBlock> MethodText;
 	TSharedPtr<STextBlock> FormatText;
 	TSharedPtr<STextBlock> LODBiasText;
+	TSharedPtr<STextBlock> HasAlphaChannelText;
+	TSharedPtr<STextBlock> NumMipsText;
 
 	/** If true, displays the red channel */
 	bool bIsRedChannel;
diff --git a/Engine/Source/Runtime/CoreUObject/Private/UObject/Obj.cpp b/Engine/Source/Runtime/CoreUObject/Private/UObject/Obj.cpp
--- a/Engine/Source/Runtime/CoreUObject/Private/UObject/Obj.cpp
+++ b/Engine/Source/Runtime/CoreUObject/Private/UObject/Obj.cpp
@@ -828,11 +828,11 @@ void UObject::Serialize( FArchive& Ar )
 	}
 
 	// Serialize object properties which are defined in the class.
-		// Handle derived UClass objects (exact UClass objects are native only and shouldn't be touched)
-		if (Class != UClass::StaticClass())
-		{
-				SerializeScriptProperties(Ar);
-			}
+	// Handle derived UClass objects (exact UClass objects are native only and shouldn't be touched)
+	if (Class != UClass::StaticClass())
+	{
+		SerializeScriptProperties(Ar);
+	}
 
 	// Keep track of pending kill
 	if( Ar.IsTransacting() )
@@ -1169,7 +1169,7 @@ void UObject::GetAssetRegistryTags(TArray<FAssetRegistryTag>& OutTags) const
 	SIZE_T ResourceSize = const_cast<UObject*>(this)->GetResourceSize(EResourceSizeMode::Exclusive);
 	if ( ResourceSize != RESOURCE_SIZE_NONE )
 	{
-		OutTags.Add( FAssetRegistryTag("ResourceSize", FString::Printf(TEXT("%0.2f"), ResourceSize / 1024.f), FAssetRegistryTag::TT_Numerical) );
+		OutTags.Add( FAssetRegistryTag("ResourceSize", FString::Printf(TEXT("%d"), (ResourceSize + 512) / 1024), FAssetRegistryTag::TT_Numerical) );
 	}
 	FAssetRegistryTag::GetAssetRegistryTagsFromSearchableProperties(this, OutTags);
 }
@@ -3613,9 +3613,9 @@ void StaticExit()
 		UObject* Obj = *It;
 		if (Obj && !Obj->IsA<UField>()) // Skip Structures, properties, etc.. They could be still necessary while GC.
 		{
-		// Mark as unreachable so purge phase will kill it.
+			// Mark as unreachable so purge phase will kill it.
 			It->SetFlags(RF_Unreachable);
-	}
+		}
 	}
 
 	// Fully purge all objects, not using time limit.
diff --git a/Engine/Source/Runtime/Engine/Private/Texture2D.cpp b/Engine/Source/Runtime/Engine/Private/Texture2D.cpp
--- a/Engine/Source/Runtime/Engine/Private/Texture2D.cpp
+++ b/Engine/Source/Runtime/Engine/Private/Texture2D.cpp
@@ -380,7 +380,7 @@ void UTexture2D::GetAssetRegistryTags(TArray<FAssetRegistryTag>& OutTags) const
 	const FString DimensionsStr = FString::Printf(TEXT("%dx%d"), SizeX, SizeY);
 	OutTags.Add( FAssetRegistryTag("Dimensions", DimensionsStr, FAssetRegistryTag::TT_Dimensional) );
 	OutTags.Add( FAssetRegistryTag("HasAlphaChannel", HasAlphaChannel() ? TEXT("True") : TEXT("False"), FAssetRegistryTag::TT_Alphabetical) );
-	
+	OutTags.Add( FAssetRegistryTag("Format", GPixelFormats[GetPixelFormat()].Name, FAssetRegistryTag::TT_Alphabetical) );
 
 	Super::GetAssetRegistryTags(OutTags);
 }
@@ -552,6 +552,12 @@ bool UTexture2D::UpdateStreamingStatus( bool bWaitForMipFading /*= false*/ )
 			else
 			{
 				ResidentMips = RequestedMips;
+#if WITH_EDITOR
+				// When all the requested mips are streamed in, generate an empty property changed event, to force the
+				// ResourceSize asset registry tag to be recalculated.
+				FPropertyChangedEvent EmptyPropertyChangedEvent(nullptr);
+				FCoreUObjectDelegates::OnObjectPropertyChanged.Broadcast(this, EmptyPropertyChangedEvent);
+#endif
 			}
 			bHasPendingRequestInFlight = false;
 
diff --git a/Engine/Source/Runtime/Engine/Private/TextureCube.cpp b/Engine/Source/Runtime/Engine/Private/TextureCube.cpp
--- a/Engine/Source/Runtime/Engine/Private/TextureCube.cpp
+++ b/Engine/Source/Runtime/Engine/Private/TextureCube.cpp
@@ -59,6 +59,7 @@ void UTextureCube::GetAssetRegistryTags(TArray<FAssetRegistryTag>& OutTags) cons
 
 	const FString Dimensions = FString::Printf(TEXT("%dx%d"), SizeX, SizeY);
 	OutTags.Add( FAssetRegistryTag("Dimensions", Dimensions, FAssetRegistryTag::TT_Dimensional) );
+	OutTags.Add( FAssetRegistryTag("Format", GPixelFormats[GetPixelFormat()].Name, FAssetRegistryTag::TT_Alphabetical) );
 
 	Super::GetAssetRegistryTags(OutTags);
 }
