Fixed clang build error.

diff --git a/Engine/Source/Editor/TextureEditor/Private/TextureEditorToolkit.cpp b/Engine/Source/Editor/TextureEditor/Private/TextureEditorToolkit.cpp
--- a/Engine/Source/Editor/TextureEditor/Private/TextureEditorToolkit.cpp
+++ b/Engine/Source/Editor/TextureEditor/Private/TextureEditorToolkit.cpp
@@ -361,7 +361,7 @@ void FTextureEditorToolkit::PopulateQuickInfo( )
 	CalculateEffectiveTextureDimensions(MipLevel, PreviewEffectiveTextureWidth, PreviewEffectiveTextureHeight);
 
 	// Texture asset size
-	const SIZE_T Size = Texture->GetResourceSize(EResourceSizeMode::Exclusive);
+	const uint32 Size = (Texture->GetResourceSize(EResourceSizeMode::Exclusive) + 512) / 1024;
 
 	FNumberFormattingOptions SizeOptions;
 	SizeOptions.UseGrouping = false;
@@ -377,7 +377,7 @@ void FTextureEditorToolkit::PopulateQuickInfo( )
 	ImportedText->SetText(FText::Format( NSLOCTEXT("TextureEditor", "QuickInfo_Imported", "Imported: {0}x{1}"), FText::AsNumber(ImportedWidth, &Options), FText::AsNumber(ImportedHeight, &Options)));
 	CurrentText->SetText(FText::Format( NSLOCTEXT("TextureEditor", "QuickInfo_Displayed", "Displayed: {0}x{1}"), FText::AsNumber(FMath::Max((uint32)1, ActualWidth >> MipLevel), &Options ), FText::AsNumber(FMath::Max((uint32)1, ActualHeight >> MipLevel), &Options)));
 	MaxInGameText->SetText(FText::Format( NSLOCTEXT("TextureEditor", "QuickInfo_MaxInGame", "Max In-Game: {0}x{1}"), FText::AsNumber(MaxInGameWidth, &Options), FText::AsNumber(MaxInGameHeight, &Options)));
-	SizeText->SetText(FText::Format(NSLOCTEXT("TextureEditor", "QuickInfo_ResourceSize", "Resource Size: {0}kb"), FText::AsNumber((Size + 512) / 1024, &SizeOptions)));
+	SizeText->SetText(FText::Format(NSLOCTEXT("TextureEditor", "QuickInfo_ResourceSize", "Resource Size: {0} Kb"), FText::AsNumber(Size, &SizeOptions)));
 	MethodText->SetText(FText::Format(NSLOCTEXT("TextureEditor", "QuickInfo_Method", "Method: {0}"), Texture->NeverStream ? NSLOCTEXT("TextureEditor", "QuickInfo_MethodNotStreamed", "Not Streamed") : NSLOCTEXT("TextureEditor", "QuickInfo_MethodStreamed", "Streamed")));
 	LODBiasText->SetText(FText::Format(NSLOCTEXT("TextureEditor", "QuickInfo_LODBias", "Combined LOD Bias: {0}"), FText::AsNumber(Texture->GetCachedLODBias())));
 
