Added a details customization for UTexture. The "Maximum Texture Size" property is specialized so that it yields an error dialog when attempting to change it for non power of two textures which do not have a padded Power of Two Mode.

diff --git a/Engine/Source/Editor/TextureEditor/Private/Customizations/TextureDetailsCustomization.cpp b/Engine/Source/Editor/TextureEditor/Private/Customizations/TextureDetailsCustomization.cpp
new file mode 100644
--- /dev/null
+++ b/Engine/Source/Editor/TextureEditor/Private/Customizations/TextureDetailsCustomization.cpp
@@ -0,0 +1,248 @@
+// Copyright 1998-2015 Epic Games, Inc. All Rights Reserved.
+
+#include "TextureEditorPrivatePCH.h"
+#include "Customizations/TextureDetailsCustomization.h"
+#include "PropertyEditorModule.h"
+#include "DetailLayoutBuilder.h"
+#include "DetailCategoryBuilder.h"
+#include "DetailWidgetRow.h"
+#include "IDetailPropertyRow.h"
+#include "PropertyHandle.h"
+#include "SNumericEntryBox.h"
+#include "STextComboBox.h"
+
+#define LOCTEXT_NAMESPACE "FTextureDetails"
+
+
+TSharedRef<IDetailCustomization> FTextureDetails::MakeInstance()
+{
+	return MakeShareable(new FTextureDetails);
+}
+
+void FTextureDetails::CustomizeDetails(IDetailLayoutBuilder& DetailBuilder)
+{
+	TArray< TWeakObjectPtr<UObject> > ObjectsBeingCustomized;
+	DetailBuilder.GetObjectsBeingCustomized(ObjectsBeingCustomized);
+	if (ensure(ObjectsBeingCustomized.Num() == 1))
+	{
+		TextureBeingCustomized = ObjectsBeingCustomized[0];
+	}
+
+	MaxTextureSizePropertyHandle = DetailBuilder.GetProperty(GET_MEMBER_NAME_CHECKED(UTexture, MaxTextureSize));
+	PowerOfTwoModePropertyHandle = DetailBuilder.GetProperty(GET_MEMBER_NAME_CHECKED(UTexture, PowerOfTwoMode));
+
+	// Customize MaxTextureSize
+	{
+		IDetailCategoryBuilder& CompressionCategory = DetailBuilder.EditCategory("Compression");
+		IDetailPropertyRow& MaxTextureSizePropertyRow = CompressionCategory.AddProperty(GET_MEMBER_NAME_CHECKED(UTexture, MaxTextureSize));
+		TSharedPtr<SWidget> NameWidget;
+		TSharedPtr<SWidget> ValueWidget;
+		FDetailWidgetRow Row;
+		MaxTextureSizePropertyRow.GetDefaultWidgets(NameWidget, ValueWidget, Row);
+
+		const bool bShowChildren = true;
+		MaxTextureSizePropertyRow.CustomWidget(bShowChildren)
+			.NameContent()
+			.MinDesiredWidth(Row.NameWidget.MinWidth)
+			.MaxDesiredWidth(Row.NameWidget.MaxWidth)
+			[
+				NameWidget.ToSharedRef()
+			]
+			.ValueContent()
+			.MinDesiredWidth(Row.ValueWidget.MinWidth)
+			.MaxDesiredWidth(Row.ValueWidget.MaxWidth)
+			[
+				SNew(SNumericEntryBox<int32>)
+				.AllowSpin(true)
+				.Value(this, &FTextureDetails::OnGetMaxTextureSize)
+				.Font(IDetailLayoutBuilder::GetDetailFont())
+				.MinValue(0)
+				.MaxValue(2048)
+				.MinSliderValue(0)
+				.MaxSliderValue(2048)
+				.OnValueChanged(this, &FTextureDetails::OnMaxTextureSizeChanged)
+				.OnValueCommitted(this, &FTextureDetails::OnMaxTextureSizeCommitted)
+				.OnBeginSliderMovement(this, &FTextureDetails::OnBeginSliderMovement)
+				.OnEndSliderMovement(this, &FTextureDetails::OnEndSliderMovement)
+			];
+	}
+
+	// Customize PowerOfTwoMode
+	{
+		IDetailCategoryBuilder& TextureCategory = DetailBuilder.EditCategory("Texture");
+		IDetailPropertyRow& PowerOfTwoModePropertyRow = TextureCategory.AddProperty(GET_MEMBER_NAME_CHECKED(UTexture, PowerOfTwoMode));
+
+		// Generate a list of enum values for the combo box
+		TArray<FText> PowerOfTwoModeComboBoxToolTips;
+		TArray<bool> RestrictedList;
+		PowerOfTwoModePropertyHandle->GeneratePossibleValues(PowerOfTwoModeComboBoxList, PowerOfTwoModeComboBoxToolTips, RestrictedList);
+
+		uint8 PowerOfTwoMode;
+		ensure(PowerOfTwoModePropertyHandle->GetValue(PowerOfTwoMode) == FPropertyAccess::Success);
+
+		TSharedPtr<SWidget> NameWidget;
+		TSharedPtr<SWidget> ValueWidget;
+		FDetailWidgetRow Row;
+		PowerOfTwoModePropertyRow.GetDefaultWidgets(NameWidget, ValueWidget, Row);
+
+		const bool bShowChildren = true;
+		PowerOfTwoModePropertyRow.CustomWidget(bShowChildren)
+			.NameContent()
+			.MinDesiredWidth(Row.NameWidget.MinWidth)
+			.MaxDesiredWidth(Row.NameWidget.MaxWidth)
+			[
+				NameWidget.ToSharedRef()
+			]
+			.ValueContent()
+			.MinDesiredWidth(Row.ValueWidget.MinWidth)
+			.MaxDesiredWidth(Row.ValueWidget.MaxWidth)
+			[
+				SAssignNew(TextComboBox, STextComboBox)
+				.Font(IDetailLayoutBuilder::GetDetailFont())
+				.OptionsSource(&PowerOfTwoModeComboBoxList)
+				.InitiallySelectedItem(PowerOfTwoModeComboBoxList[PowerOfTwoMode])
+				.OnSelectionChanged(this, &FTextureDetails::OnPowerOfTwoModeChanged)
+			];
+	}
+}
+
+bool FTextureDetails::CanEditMaxTextureSize() const
+{
+	if (UTexture2D* Texture2D = Cast<UTexture2D>(TextureBeingCustomized.Get()))
+	{
+		if (!Texture2D->Source.IsPowerOfTwo() && Texture2D->PowerOfTwoMode == ETexturePowerOfTwoSetting::None)
+		{
+			return false;
+		}
+	}
+
+	return true;
+}
+
+void FTextureDetails::CreateMaxTextureSizeMessage() const
+{
+	FMessageDialog::Open(EAppMsgType::Ok,
+		LOCTEXT("CannotEditMaxTextureSize", "Maximum Texture Size cannot be changed for this texture as it is a non power of two size. Change the Power of Two Mode to allow it to be padded to a power of two."));
+}
+
+/** @return The value or unset if properties with multiple values are viewed */
+TOptional<int32> FTextureDetails::OnGetMaxTextureSize() const
+{
+	int32 NumericVal;
+	if (MaxTextureSizePropertyHandle->GetValue(NumericVal) == FPropertyAccess::Success)
+	{
+		return NumericVal;
+	}
+
+	// Return an unset value so it displays the "multiple values" indicator instead
+	return TOptional<int32>();
+}
+
+void FTextureDetails::OnMaxTextureSizeChanged(int32 NewValue)
+{
+	if (!CanEditMaxTextureSize())
+	{
+		return;
+	}
+
+	if (bIsUsingSlider)
+	{
+		int32 OrgValue(0);
+		if (MaxTextureSizePropertyHandle->GetValue(OrgValue) != FPropertyAccess::Fail)
+		{
+			// Value hasn't changed, so let's return now
+			if (OrgValue == NewValue)
+			{
+				return;
+			}
+		}
+
+		// We don't create a transaction for each property change when using the slider.  Only once when the slider first is moved
+		EPropertyValueSetFlags::Type Flags = (EPropertyValueSetFlags::InteractiveChange | EPropertyValueSetFlags::NotTransactable);
+		MaxTextureSizePropertyHandle->SetValue(NewValue, Flags);
+	}
+}
+
+void FTextureDetails::OnMaxTextureSizeCommitted(int32 NewValue, ETextCommit::Type CommitInfo)
+{
+	if (!CanEditMaxTextureSize())
+	{
+		if (CommitInfo == ETextCommit::OnEnter)
+		{
+			CreateMaxTextureSizeMessage();
+		}
+		return;
+	}
+
+	MaxTextureSizePropertyHandle->SetValue(NewValue);
+}
+
+/**
+ * Called when the slider begins to move.  We create a transaction here to undo the property
+ */
+void FTextureDetails::OnBeginSliderMovement()
+{
+	if (!CanEditMaxTextureSize())
+	{
+		return;
+	}
+
+	bIsUsingSlider = true;
+
+	GEditor->BeginTransaction(TEXT("TextureDetails"), LOCTEXT("SetMaximumTextureSize", "Edit Maximum Texture Size"), MaxTextureSizePropertyHandle->GetProperty());
+}
+
+
+/**
+ * Called when the slider stops moving.  We end the previously created transaction
+ */
+void FTextureDetails::OnEndSliderMovement(int32 NewValue)
+{
+	if (!CanEditMaxTextureSize())
+	{
+		return;
+	}
+
+	bIsUsingSlider = false;
+
+	GEditor->EndTransaction();
+}
+
+bool FTextureDetails::CanEditPowerOfTwoMode(int32 NewPowerOfTwoMode) const
+{
+	if (UTexture2D* Texture2D = Cast<UTexture2D>(TextureBeingCustomized.Get()))
+	{
+		if (!Texture2D->Source.IsPowerOfTwo() && Texture2D->MaxTextureSize > 0 && NewPowerOfTwoMode == ETexturePowerOfTwoSetting::None)
+		{
+			return false;
+		}
+	}
+
+	return true;
+}
+
+void FTextureDetails::CreatePowerOfTwoModeMessage() const
+{
+	FMessageDialog::Open(EAppMsgType::Ok,
+						 LOCTEXT("CannotEditPowerOfTwoMode", "Power of Two Mode cannot be changed to None for this texture as it is a non power of two size and has a Maximum Texture Size override. Change the Maximum Texture Size to 0 before attempting to change the Power of Two Mode."));
+}
+
+void FTextureDetails::OnPowerOfTwoModeChanged(TSharedPtr<FString> NewValue, ESelectInfo::Type SelectInfo)
+{
+	int32 NewPowerOfTwoMode = PowerOfTwoModeComboBoxList.Find(NewValue);
+	check(NewPowerOfTwoMode != INDEX_NONE);
+	if (!CanEditPowerOfTwoMode(NewPowerOfTwoMode))
+	{
+		CreatePowerOfTwoModeMessage();
+		uint8 CurrentPowerOfTwoMode;
+		ensure(PowerOfTwoModePropertyHandle->GetValue(CurrentPowerOfTwoMode) == FPropertyAccess::Success);
+		TextComboBox->SetSelectedItem(PowerOfTwoModeComboBoxList[CurrentPowerOfTwoMode]);
+		return;
+	}
+
+	int32 PowerOfTwoMode = PowerOfTwoModeComboBoxList.Find(NewValue);
+	check(PowerOfTwoMode != INDEX_NONE);
+	ensure(PowerOfTwoModePropertyHandle->SetValue(static_cast<uint8>(PowerOfTwoMode)) == FPropertyAccess::Success);
+}
+
+#undef LOCTEXT_NAMESPACE
diff --git a/Engine/Source/Editor/TextureEditor/Private/Customizations/TextureDetailsCustomization.h b/Engine/Source/Editor/TextureEditor/Private/Customizations/TextureDetailsCustomization.h
new file mode 100644
--- /dev/null
+++ b/Engine/Source/Editor/TextureEditor/Private/Customizations/TextureDetailsCustomization.h
@@ -0,0 +1,38 @@
+// Copyright 1998-2015 Epic Games, Inc. All Rights Reserved.
+
+#pragma once
+
+
+class FTextureDetails : public IDetailCustomization
+{
+public:
+	/** Makes a new instance of this detail layout class for a specific detail view requesting it */
+	static TSharedRef<IDetailCustomization> MakeInstance();
+
+private:
+	FTextureDetails()
+		: bIsUsingSlider(false)
+	{}
+
+	/** IDetailCustomization interface */
+	virtual void CustomizeDetails(IDetailLayoutBuilder& DetailBuilder) override;
+
+	TOptional<int32> OnGetMaxTextureSize() const;
+	void OnMaxTextureSizeChanged(int32 NewValue);
+	void OnMaxTextureSizeCommitted(int32 NewValue, ETextCommit::Type CommitInfo);
+	void OnBeginSliderMovement();
+	void OnEndSliderMovement(int32 NewValue);
+	bool CanEditMaxTextureSize() const;
+	void CreateMaxTextureSizeMessage() const;
+	void OnPowerOfTwoModeChanged(TSharedPtr<FString> NewValue, ESelectInfo::Type SelectInfo);
+	bool CanEditPowerOfTwoMode(int32 NewPowerOfTwoMode) const;
+	void CreatePowerOfTwoModeMessage() const;
+
+	TSharedPtr<IPropertyHandle> MaxTextureSizePropertyHandle;
+	TSharedPtr<IPropertyHandle> PowerOfTwoModePropertyHandle;
+	TArray<TSharedPtr<FString>> PowerOfTwoModeComboBoxList;
+	TWeakObjectPtr<UObject> TextureBeingCustomized;
+	TSharedPtr<STextComboBox> TextComboBox;
+	bool bIsUsingSlider;
+};
+
diff --git a/Engine/Source/Editor/TextureEditor/Private/TextureEditorModule.cpp b/Engine/Source/Editor/TextureEditor/Private/TextureEditorModule.cpp
--- a/Engine/Source/Editor/TextureEditor/Private/TextureEditorModule.cpp
+++ b/Engine/Source/Editor/TextureEditor/Private/TextureEditorModule.cpp
@@ -4,10 +4,19 @@
 #include "TextureEditorPrivatePCH.h"
 #include "ISettingsModule.h"
 #include "ModuleManager.h"
+#include "Customizations/TextureDetailsCustomization.h"
 
+#include "PropertyEditorModule.h"
+#include "DetailLayoutBuilder.h"
+#include "DetailCategoryBuilder.h"
+#include "DetailWidgetRow.h"
+#include "IDetailPropertyRow.h"
+#include "PropertyHandle.h"
+#include "SNumericEntryBox.h"
 
 #define LOCTEXT_NAMESPACE "FTextureEditorModule"
 
+
 const FName TextureEditorAppIdentifier = FName(TEXT("TextureEditorApp"));
 
 
@@ -61,6 +70,9 @@ public:
 				GetMutableDefault<UTextureEditorSettings>()
 			);
 		}
+
+		FPropertyEditorModule& PropertyModule = FModuleManager::GetModuleChecked<FPropertyEditorModule>("PropertyEditor");
+		PropertyModule.RegisterCustomClassLayout("Texture", FOnGetDetailCustomizationInstance::CreateStatic(&FTextureDetails::MakeInstance));
 	}
 
 	virtual void ShutdownModule( ) override
@@ -76,6 +88,9 @@ public:
 		// unregister menu extensions
 		MenuExtensibilityManager.Reset();
 		ToolBarExtensibilityManager.Reset();
+
+		FPropertyEditorModule& PropertyModule = FModuleManager::GetModuleChecked<FPropertyEditorModule>("PropertyEditor");
+		PropertyModule.UnregisterCustomClassLayout("Texture2D");
 	}
 
 private:
diff --git a/Engine/Source/Editor/TextureEditor/TextureEditor.Build.cs b/Engine/Source/Editor/TextureEditor/TextureEditor.Build.cs
--- a/Engine/Source/Editor/TextureEditor/TextureEditor.Build.cs
+++ b/Engine/Source/Editor/TextureEditor/TextureEditor.Build.cs
@@ -27,6 +27,7 @@ public class TextureEditor : ModuleRules
 			new string[] {
 				"Settings",
 				"UnrealEd",
+                "PropertyEditor"
 			}
 		);
 
diff --git a/Engine/Source/Runtime/Engine/Private/Texture.cpp b/Engine/Source/Runtime/Engine/Private/Texture.cpp
--- a/Engine/Source/Runtime/Engine/Private/Texture.cpp
+++ b/Engine/Source/Runtime/Engine/Private/Texture.cpp
@@ -126,7 +126,6 @@ void UTexture::PostEditChangeProperty(FPropertyChangedEvent& PropertyChangedEven
 		static const FName DeferCompressionName("DeferCompression");
 #if WITH_EDITORONLY_DATA
 		static const FName MaxTextureSizeName("MaxTextureSize");
-		static const FName PowerOfTwoModeName("PowerOfTwoMode");
 #endif // #if WITH_EDITORONLY_DATA
 
 		const FName PropertyName = PropertyThatChanged->GetFName();
@@ -147,23 +146,7 @@ void UTexture::PostEditChangeProperty(FPropertyChangedEvent& PropertyChangedEven
 			}
 			else
 			{
-				if (!Source.IsPowerOfTwo() && PowerOfTwoMode == ETexturePowerOfTwoSetting::None)
-				{
-					// Force power-of-two padding if a max texture size is specified on a non power-of-two sized texture
-					PowerOfTwoMode = ETexturePowerOfTwoSetting::PadToPowerOfTwo;
-				}
-				else
-				{
-					MaxTextureSize = FMath::Min<int32>(FMath::RoundUpToPowerOfTwo(MaxTextureSize), GetMaximumDimension());
-				}
-			}
-		}
-		else if (PropertyName == PowerOfTwoModeName)
-		{
-			if (PowerOfTwoMode == ETexturePowerOfTwoSetting::None && MaxTextureSize > 0 && !Source.IsPowerOfTwo())
-			{
-				// Force no max texture size if the source texture is not power-of-two sized and we are requesting no power-of-two padding
-				MaxTextureSize = 0;
+				MaxTextureSize = FMath::Min<int32>(FMath::RoundUpToPowerOfTwo(MaxTextureSize), GetMaximumDimension());
 			}
 		}
 #endif // #if WITH_EDITORONLY_DATA
