Fixed crash with new texture customization for MaxSize and PowerOfTwo mode.  Some textures do not have these settings and the customization did not account for that

https://jira.ol.epicgames.net/browse/UE-19641

diff --git a/Engine/Source/Editor/TextureEditor/Private/Customizations/TextureDetailsCustomization.cpp b/Engine/Source/Editor/TextureEditor/Private/Customizations/TextureDetailsCustomization.cpp
--- a/Engine/Source/Editor/TextureEditor/Private/Customizations/TextureDetailsCustomization.cpp
+++ b/Engine/Source/Editor/TextureEditor/Private/Customizations/TextureDetailsCustomization.cpp
@@ -32,6 +32,7 @@ void FTextureDetails::CustomizeDetails(IDetailLayoutBuilder& DetailBuilder)
 	PowerOfTwoModePropertyHandle = DetailBuilder.GetProperty(GET_MEMBER_NAME_CHECKED(UTexture, PowerOfTwoMode));
 
 	// Customize MaxTextureSize
+	if( MaxTextureSizePropertyHandle->IsValidHandle() )
 	{
 		IDetailCategoryBuilder& CompressionCategory = DetailBuilder.EditCategory("Compression");
 		IDetailPropertyRow& MaxTextureSizePropertyRow = CompressionCategory.AddProperty(GET_MEMBER_NAME_CHECKED(UTexture, MaxTextureSize));
@@ -68,6 +69,7 @@ void FTextureDetails::CustomizeDetails(IDetailLayoutBuilder& DetailBuilder)
 	}
 
 	// Customize PowerOfTwoMode
+	if( PowerOfTwoModePropertyHandle->IsValidHandle() )
 	{
 		IDetailCategoryBuilder& TextureCategory = DetailBuilder.EditCategory("Texture");
 		IDetailPropertyRow& PowerOfTwoModePropertyRow = TextureCategory.AddProperty(GET_MEMBER_NAME_CHECKED(UTexture, PowerOfTwoMode));
