Removed redundant calls to InvalidateDisplay() which is already called by Invalidate().

diff --git a/Engine/Source/Editor/FontEditor/Private/SFontEditorViewport.cpp b/Engine/Source/Editor/FontEditor/Private/SFontEditorViewport.cpp
--- a/Engine/Source/Editor/FontEditor/Private/SFontEditorViewport.cpp
+++ b/Engine/Source/Editor/FontEditor/Private/SFontEditorViewport.cpp
@@ -438,7 +438,6 @@ bool FFontEditorViewportClient::InputKey(FViewport* Viewport, int32 ControllerId
 		
 			// Force a redraw
 			Viewport->Invalidate();
-			Viewport->InvalidateDisplay();
 
 			bHandled = true;
 		}
diff --git a/Engine/Source/Editor/TextureEditor/Private/Widgets/STextureEditorViewport.cpp b/Engine/Source/Editor/TextureEditor/Private/Widgets/STextureEditorViewport.cpp
--- a/Engine/Source/Editor/TextureEditor/Private/Widgets/STextureEditorViewport.cpp
+++ b/Engine/Source/Editor/TextureEditor/Private/Widgets/STextureEditorViewport.cpp
@@ -334,7 +334,6 @@ void STextureEditorViewport::Tick( const FGeometry& AllottedGeometry, const doub
 	if (bIsRenderingEnabled)
 	{
 		Viewport->Invalidate();
-		Viewport->InvalidateDisplay();
 	}
 }
 
