Fixed possible division by 0 warnings.

diff --git a/Engine/Plugins/Lumin/MagicLeap/Source/MagicLeapAR/Private/LuminARTrackingSystem.cpp b/Engine/Plugins/Lumin/MagicLeap/Source/MagicLeapAR/Private/LuminARTrackingSystem.cpp
--- a/Engine/Plugins/Lumin/MagicLeap/Source/MagicLeapAR/Private/LuminARTrackingSystem.cpp
+++ b/Engine/Plugins/Lumin/MagicLeap/Source/MagicLeapAR/Private/LuminARTrackingSystem.cpp
@@ -743,7 +743,7 @@ ULuminARCandidateImage* FLuminARImplementation::AddLuminRuntimeCandidateImage(UA
 {
 	if (OnAddRuntimeCandidateImage(SessionConfig, CandidateTexture, FriendlyName, PhysicalWidth))
 	{
-		float PhysicalHeight = PhysicalWidth / CandidateTexture->GetSizeX() * CandidateTexture->GetSizeY();
+		float PhysicalHeight = PhysicalWidth / FMath::Max<int32>(1, CandidateTexture->GetSizeX()) * CandidateTexture->GetSizeY();
 		ULuminARCandidateImage* NewCandidateImage = ULuminARCandidateImage::CreateNewLuminARCandidateImage(CandidateTexture, FriendlyName, PhysicalWidth, PhysicalHeight, EARCandidateImageOrientation::Landscape, bUseUnreliablePose, bImageIsStationary);
 		SessionConfig->AddCandidateImage(NewCandidateImage);
 		return NewCandidateImage;
diff --git a/Engine/Plugins/Runtime/Oculus/OculusVR/Source/OculusMR/Private/OculusMR_CastingCameraActor.cpp b/Engine/Plugins/Runtime/Oculus/OculusVR/Source/OculusMR/Private/OculusMR_CastingCameraActor.cpp
--- a/Engine/Plugins/Runtime/Oculus/OculusVR/Source/OculusMR/Private/OculusMR_CastingCameraActor.cpp
+++ b/Engine/Plugins/Runtime/Oculus/OculusVR/Source/OculusMR/Private/OculusMR_CastingCameraActor.cpp
@@ -597,7 +597,7 @@ void AOculusMR_CastingCameraActor::UpdateCameraColorTexture(const ovrpSizei &fra
 		{
 			CameraFrameMaterialInstance->SetTextureParameterValue(FName(TEXT("CameraCaptureTexture")), CameraColorTexture);
 			CameraFrameMaterialInstance->SetVectorParameterValue(FName(TEXT("CameraCaptureTextureSize")),
-				FLinearColor((float)CameraColorTexture->GetSizeX(), (float)CameraColorTexture->GetSizeY(), 1.0f / CameraColorTexture->GetSizeX(), 1.0f / CameraColorTexture->GetSizeY()));
+				FLinearColor((float)CameraColorTexture->GetSizeX(), (float)CameraColorTexture->GetSizeY(), 1.0f / FMath::Max<int32>(1, CameraColorTexture->GetSizeX()), 1.0f / FMath::Max<int32>(1, CameraColorTexture->GetSizeY())));
 		}
 	}
 	uint32 Pitch = rowPitch;
@@ -985,7 +985,7 @@ void AOculusMR_CastingCameraActor::SetupCameraFrameMaterialInstance()
 	{
 		CameraFrameMaterialInstance->SetTextureParameterValue(FName(TEXT("CameraCaptureTexture")), CameraColorTexture);
 		CameraFrameMaterialInstance->SetVectorParameterValue(FName(TEXT("CameraCaptureTextureSize")),
-			FLinearColor((float)CameraColorTexture->GetSizeX(), (float)CameraColorTexture->GetSizeY(), 1.0f / CameraColorTexture->GetSizeX(), 1.0f / CameraColorTexture->GetSizeY()));
+			FLinearColor((float)CameraColorTexture->GetSizeX(), (float)CameraColorTexture->GetSizeY(), 1.0f / FMath::Max<int32>(1, CameraColorTexture->GetSizeX()), 1.0f / FMath::Max<int32>(1, CameraColorTexture->GetSizeY())));
 		if (MRSettings->GetUseDynamicLighting())
 		{
 			CameraFrameMaterialInstance->SetTextureParameterValue(FName(TEXT("CameraDepthTexture")), CameraDepthTexture);
diff --git a/Engine/Source/Editor/LandscapeEditor/Private/LandscapeEdModeBrushes.cpp b/Engine/Source/Editor/LandscapeEditor/Private/LandscapeEdModeBrushes.cpp
--- a/Engine/Source/Editor/LandscapeEditor/Private/LandscapeEdModeBrushes.cpp
+++ b/Engine/Source/Editor/LandscapeEditor/Private/LandscapeEdModeBrushes.cpp
@@ -638,8 +638,8 @@ public:
 
 					float SquaredScaleXY = FMath::Square(ScaleXY);
 					FLinearColor AlphaScaleBias(
-						SquaredScaleXY / (Gizmo->GetWidth() * DataTexture->GetSizeX()),
-						SquaredScaleXY / (Gizmo->GetHeight() * DataTexture->GetSizeY()),
+						SquaredScaleXY / FMath::Max<float>(1.f, Gizmo->GetWidth() * DataTexture->GetSizeX()),
+						SquaredScaleXY / FMath::Max<float>(1.f, Gizmo->GetHeight() * DataTexture->GetSizeY()),
 						Gizmo->TextureScale.X,
 						Gizmo->TextureScale.Y
 						);
diff --git a/Engine/Source/Editor/TextureEditor/Private/Models/TextureEditorViewportClient.cpp b/Engine/Source/Editor/TextureEditor/Private/Models/TextureEditorViewportClient.cpp
--- a/Engine/Source/Editor/TextureEditor/Private/Models/TextureEditorViewportClient.cpp
+++ b/Engine/Source/Editor/TextureEditor/Private/Models/TextureEditorViewportClient.cpp
@@ -144,13 +144,18 @@ void FTextureEditorViewportClient::Draw(FViewport* Viewport, FCanvas* Canvas)
 
 	// Draw the background checkerboard pattern in the same size/position as the render texture so it will show up anywhere
 	// the texture has transparency
-	if (Settings.Background == TextureEditorBackground_CheckeredFill)
+	if (Viewport && CheckerboardTexture)
 	{
-		Canvas->DrawTile( 0.0f, 0.0f, Viewport->GetSizeXY().X, Viewport->GetSizeXY().Y, 0.0f, 0.0f, (Viewport->GetSizeXY().X / CheckerboardTexture->GetSizeX()), (Viewport->GetSizeXY().Y / CheckerboardTexture->GetSizeY()), FLinearColor::White, CheckerboardTexture->Resource);
-	}
-	else if (Settings.Background == TextureEditorBackground_Checkered)
-	{
-		Canvas->DrawTile( XPos, YPos, Width, Height, 0.0f, 0.0f, (Width / CheckerboardTexture->GetSizeX()), (Height / CheckerboardTexture->GetSizeY()), FLinearColor::White, CheckerboardTexture->Resource);
+		const int32 CheckerboardSizeX = FMath::Max<int32>(1, CheckerboardTexture->GetSizeX());
+		const int32 CheckerboardSizeY = FMath::Max<int32>(1, CheckerboardTexture->GetSizeY());
+		if (Settings.Background == TextureEditorBackground_CheckeredFill)
+		{
+			Canvas->DrawTile( 0.0f, 0.0f, Viewport->GetSizeXY().X, Viewport->GetSizeXY().Y, 0.0f, 0.0f, Viewport->GetSizeXY().X / CheckerboardSizeX, Viewport->GetSizeXY().Y / CheckerboardSizeY, FLinearColor::White, CheckerboardTexture->Resource);
+		}
+		else if (Settings.Background == TextureEditorBackground_Checkered)
+		{
+			Canvas->DrawTile( XPos, YPos, Width, Height, 0.0f, 0.0f, Width / CheckerboardSizeX, Height / CheckerboardSizeY, FLinearColor::White, CheckerboardTexture->Resource);
+		}
 	}
 
 	float Exposure = FMath::Pow(2.0f, (float)TextureEditorViewportPtr.Pin()->GetExposureBias());
diff --git a/Engine/Source/Runtime/AugmentedReality/Private/ARSupportInterface.cpp b/Engine/Source/Runtime/AugmentedReality/Private/ARSupportInterface.cpp
--- a/Engine/Source/Runtime/AugmentedReality/Private/ARSupportInterface.cpp
+++ b/Engine/Source/Runtime/AugmentedReality/Private/ARSupportInterface.cpp
@@ -285,7 +285,7 @@ UARCandidateImage* FARSupportInterface::AddRuntimeCandidateImage(UARSessionConfi
 {
 	if (ARImplemention && ARImplemention->OnAddRuntimeCandidateImage(SessionConfig, CandidateTexture, FriendlyName, PhysicalWidth))
 	{
-		float PhysicalHeight = PhysicalWidth / CandidateTexture->GetSizeX() * CandidateTexture->GetSizeY();
+		float PhysicalHeight = PhysicalWidth / FMath::Max<int32>(1, CandidateTexture->GetSizeX()) * CandidateTexture->GetSizeY();
 		UARCandidateImage* NewCandidateImage = UARCandidateImage::CreateNewARCandidateImage(CandidateTexture, FriendlyName, PhysicalWidth, PhysicalHeight, EARCandidateImageOrientation::Landscape);
 		SessionConfig->AddCandidateImage(NewCandidateImage);
 		return NewCandidateImage;
diff --git a/Engine/Source/Runtime/AugmentedReality/Private/ARTypes.cpp b/Engine/Source/Runtime/AugmentedReality/Private/ARTypes.cpp
--- a/Engine/Source/Runtime/AugmentedReality/Private/ARTypes.cpp
+++ b/Engine/Source/Runtime/AugmentedReality/Private/ARTypes.cpp
@@ -43,11 +43,11 @@ void UARCandidateImage::PostEditChangeProperty(FPropertyChangedEvent& PropertyCh
 		{
 			if (Orientation == EARCandidateImageOrientation::Landscape)
 			{
-				AspectRatio = (float)CandidateTexture->GetSizeY() / (float)CandidateTexture->GetSizeX();
+				AspectRatio = (float)CandidateTexture->GetSizeY() / (float)FMath::Max<int32>(1, CandidateTexture->GetSizeX());
 			}
 			else
 			{
-				AspectRatio = (float)CandidateTexture->GetSizeX() / (float)CandidateTexture->GetSizeY();
+				AspectRatio = (float)CandidateTexture->GetSizeX() / (float)FMath::Max<int32>(1, CandidateTexture->GetSizeY());
 			}
 		}
 		// If the texture has changed enforce the aspect ratio on the physical size
diff --git a/Engine/Source/Runtime/Landscape/Private/LandscapeRender.cpp b/Engine/Source/Runtime/Landscape/Private/LandscapeRender.cpp
--- a/Engine/Source/Runtime/Landscape/Private/LandscapeRender.cpp
+++ b/Engine/Source/Runtime/Landscape/Private/LandscapeRender.cpp
@@ -1281,8 +1281,8 @@ FLandscapeComponentSceneProxy::FLandscapeComponentSceneProxy(ULandscapeComponent
 	}
 	else
 	{
-		HeightmapSubsectionOffsetU = ((float)(InComponent->SubsectionSizeQuads + 1) / (float)HeightmapTexture->GetSizeX());
-		HeightmapSubsectionOffsetV = ((float)(InComponent->SubsectionSizeQuads + 1) / (float)HeightmapTexture->GetSizeY());
+		HeightmapSubsectionOffsetU = ((float)(InComponent->SubsectionSizeQuads + 1) / (float)FMath::Max<int32>(1, HeightmapTexture->GetSizeX()));
+		HeightmapSubsectionOffsetV = ((float)(InComponent->SubsectionSizeQuads + 1) / (float)FMath::Max<int32>(1, HeightmapTexture->GetSizeY()));
 	}
 
 	float ScreenSizeRatioDivider = FMath::Max(InComponent->GetLandscapeProxy()->LOD0DistributionSetting * GLandscapeLOD0DistributionScale, 1.01f);
diff --git a/Engine/Source/Runtime/Renderer/Private/BlueNoise.h b/Engine/Source/Runtime/Renderer/Private/BlueNoise.h
--- a/Engine/Source/Runtime/Renderer/Private/BlueNoise.h
+++ b/Engine/Source/Runtime/Renderer/Private/BlueNoise.h
@@ -22,6 +22,11 @@ END_GLOBAL_SHADER_PARAMETER_STRUCT()
 inline void InitializeBlueNoise(FBlueNoise& BlueNoise)
 {
 	check(GEngine);
-	BlueNoise.Dimensions = FIntVector(GEngine->BlueNoiseTexture->GetSizeX(), GEngine->BlueNoiseTexture->GetSizeX(), GEngine->BlueNoiseTexture->GetSizeY() / GEngine->BlueNoiseTexture->GetSizeX());
+	BlueNoise.Dimensions = FIntVector(
+		GEngine->BlueNoiseTexture->GetSizeX(), 
+		GEngine->BlueNoiseTexture->GetSizeX(), 
+		GEngine->BlueNoiseTexture->GetSizeY() / FMath::Max<int32>(1, GEngine->BlueNoiseTexture->GetSizeX())
+	);
+
 	BlueNoise.Texture = GEngine->BlueNoiseTexture->Resource->TextureRHI;
 }
