Texture editor reskin

diff --git a/Engine/Source/Editor/EditorStyle/Private/StarshipStyle.cpp b/Engine/Source/Editor/EditorStyle/Private/StarshipStyle.cpp
--- a/Engine/Source/Editor/EditorStyle/Private/StarshipStyle.cpp
+++ b/Engine/Source/Editor/EditorStyle/Private/StarshipStyle.cpp
@@ -6742,21 +6742,51 @@ void FStarshipEditorStyle::FStyle::SetupToolkitStyles()
 	{
 		Set("TextureEditor.Tabs.Properties", new IMAGE_BRUSH("/Icons/icon_tab_SelectionDetails_16x", Icon16x16));
 		
-		Set( "TextureEditor.RedChannel", new IMAGE_BRUSH( "Icons/icon_TextureEd_RedChannel_40x", Icon40x40 ) );
-		Set( "TextureEditor.RedChannel.Small", new IMAGE_BRUSH( "Icons/icon_TextureEd_RedChannel_40x", Icon20x20 ) );
-		Set( "TextureEditor.GreenChannel", new IMAGE_BRUSH( "Icons/icon_TextureEd_GreenChannel_40x", Icon40x40 ) );
-		Set( "TextureEditor.GreenChannel.Small", new IMAGE_BRUSH( "Icons/icon_TextureEd_GreenChannel_40x", Icon20x20 ) );
-		Set( "TextureEditor.BlueChannel", new IMAGE_BRUSH( "Icons/icon_TextureEd_BlueChannel_40x", Icon40x40 ) );
-		Set( "TextureEditor.BlueChannel.Small", new IMAGE_BRUSH( "Icons/icon_TextureEd_BlueChannel_40x", Icon20x20 ) );
-		Set( "TextureEditor.AlphaChannel", new IMAGE_BRUSH( "Icons/icon_TextureEd_AlphaChannel_40x", Icon40x40 ) );
-		Set( "TextureEditor.AlphaChannel.Small", new IMAGE_BRUSH( "Icons/icon_TextureEd_AlphaChannel_40x", Icon20x20 ) );
-		Set( "TextureEditor.Saturation", new IMAGE_BRUSH( "Icons/icon_TextureEd_Saturation_40x", Icon40x40 ) );
-		Set( "TextureEditor.Saturation.Small", new IMAGE_BRUSH( "Icons/icon_TextureEd_Saturation_40x", Icon20x20 ) );
-
-		Set( "TextureEditor.CompressNow", new IMAGE_BRUSH( "Icons/icon_TextureEd_CompressNow_40x", Icon40x40 ) );
-		Set( "TextureEditor.CompressNow.Small", new IMAGE_BRUSH( "Icons/icon_TextureEd_CompressNow_40x", Icon20x20 ) );
-		Set( "TextureEditor.Reimport", new IMAGE_BRUSH( "Icons/icon_TextureEd_Reimport_40x", Icon40x40 ) );
-		Set( "TextureEditor.Reimport.Small", new IMAGE_BRUSH( "Icons/icon_TextureEd_Reimport_40x", Icon20x20 ) );
+		Set("TextureEditor.RedChannel", new IMAGE_BRUSH( "Icons/icon_TextureEd_RedChannel_40x", Icon40x40));
+		Set("TextureEditor.RedChannel.Small", new IMAGE_BRUSH( "Icons/icon_TextureEd_RedChannel_40x", Icon20x20));
+		Set("TextureEditor.GreenChannel", new IMAGE_BRUSH( "Icons/icon_TextureEd_GreenChannel_40x", Icon40x40));
+		Set("TextureEditor.GreenChannel.Small", new IMAGE_BRUSH( "Icons/icon_TextureEd_GreenChannel_40x", Icon20x20));
+		Set("TextureEditor.BlueChannel", new IMAGE_BRUSH( "Icons/icon_TextureEd_BlueChannel_40x", Icon40x40));
+		Set("TextureEditor.BlueChannel.Small", new IMAGE_BRUSH( "Icons/icon_TextureEd_BlueChannel_40x", Icon20x20));
+		Set("TextureEditor.AlphaChannel", new IMAGE_BRUSH( "Icons/icon_TextureEd_AlphaChannel_40x", Icon40x40));
+		Set("TextureEditor.AlphaChannel.Small", new IMAGE_BRUSH( "Icons/icon_TextureEd_AlphaChannel_40x", Icon20x20));
+		Set("TextureEditor.Saturation", new IMAGE_BRUSH( "Icons/icon_TextureEd_Saturation_40x", Icon40x40));
+		Set("TextureEditor.Saturation.Small", new IMAGE_BRUSH( "Icons/icon_TextureEd_Saturation_40x", Icon20x20));
+
+		Set("TextureEditor.CompressNow", new IMAGE_BRUSH( "Icons/icon_TextureEd_CompressNow_40x", Icon40x40));
+		Set("TextureEditor.CompressNow.Small", new IMAGE_BRUSH( "Icons/icon_TextureEd_CompressNow_40x", Icon20x20));
+		Set("TextureEditor.Reimport", new IMAGE_BRUSH( "Icons/icon_TextureEd_Reimport_40x", Icon40x40));
+		Set("TextureEditor.Reimport.Small", new IMAGE_BRUSH( "Icons/icon_TextureEd_Reimport_40x", Icon20x20));
+
+		FButtonStyle MipmapButtonStyle = 
+			FButtonStyle(FStarshipCoreStyle::GetCoreStyle().GetWidgetStyle<FButtonStyle>("Button"))
+				.SetNormalPadding(FMargin(2, 2, 2, 2))
+				.SetPressedPadding(FMargin(2, 3, 2, 1));
+
+		Set("TextureEditor.MipmapButtonStyle", MipmapButtonStyle);
+
+		const FLinearColor White80 = FLinearColor(1, 1, 1, .8f);
+
+		const FCheckBoxStyle ChannelToggleButtonStyle = FCheckBoxStyle()
+			.SetCheckBoxType(ESlateCheckBoxType::ToggleButton)
+			.SetCheckedImage(FSlateRoundedBoxBrush(White80, CoreStyleConstants::InputFocusRadius))
+			.SetCheckedHoveredImage(FSlateRoundedBoxBrush(FLinearColor::White, CoreStyleConstants::InputFocusRadius))
+			.SetCheckedPressedImage(FSlateRoundedBoxBrush(FStyleColors::White, CoreStyleConstants::InputFocusRadius))
+			.SetUncheckedImage(FSlateRoundedBoxBrush(FStyleColors::Dropdown, CoreStyleConstants::InputFocusRadius))
+			.SetUncheckedHoveredImage(FSlateRoundedBoxBrush(FStyleColors::Hover, CoreStyleConstants::InputFocusRadius))
+			.SetUncheckedPressedImage(FSlateRoundedBoxBrush(FStyleColors::Hover, CoreStyleConstants::InputFocusRadius))
+			.SetForegroundColor(FStyleColors::Foreground)
+			.SetHoveredForegroundColor(FStyleColors::ForegroundHover)
+			.SetPressedForegroundColor(FStyleColors::ForegroundHover)
+			.SetCheckedForegroundColor(FStyleColors::Foreground)
+			.SetCheckedHoveredForegroundColor(FStyleColors::ForegroundHover)
+			.SetPadding(FMargin(8.f, 4.f));
+
+		FSlateFontInfo ChannelButtonFont = FStyleFonts::Get().NormalBold;
+		ChannelButtonFont.Size = 12;
+		Set("TextureEditor.ChannelButtonFont", ChannelButtonFont);
+
+		Set("TextureEditor.ChannelButtonStyle", ChannelToggleButtonStyle);
 	}
 
 	// Cascade
diff --git a/Engine/Source/Editor/TextureEditor/Private/Customizations/TextureDetailsCustomization.cpp b/Engine/Source/Editor/TextureEditor/Private/Customizations/TextureDetailsCustomization.cpp
--- a/Engine/Source/Editor/TextureEditor/Private/Customizations/TextureDetailsCustomization.cpp
+++ b/Engine/Source/Editor/TextureEditor/Private/Customizations/TextureDetailsCustomization.cpp
@@ -30,6 +30,12 @@ void FTextureDetails::CustomizeDetails(IDetailLayoutBuilder& DetailBuilder)
 		TextureBeingCustomized = ObjectsBeingCustomized[0];
 	}
 
+	DetailBuilder.EditCategory("LevelOfDetail");
+	DetailBuilder.EditCategory("Compression");
+	DetailBuilder.EditCategory("Texture");
+	DetailBuilder.EditCategory("Adjustments");
+	DetailBuilder.EditCategory("File Path");
+
 	MaxTextureSizePropertyHandle = DetailBuilder.GetProperty(GET_MEMBER_NAME_CHECKED(UTexture, MaxTextureSize));
 	PowerOfTwoModePropertyHandle = DetailBuilder.GetProperty(GET_MEMBER_NAME_CHECKED(UTexture, PowerOfTwoMode));
 	VirtualTextureStreamingPropertyHandle = DetailBuilder.GetProperty(GET_MEMBER_NAME_CHECKED(UTexture, VirtualTextureStreaming));
diff --git a/Engine/Source/Editor/TextureEditor/Private/Menus/TextureEditorViewOptionsMenu.h b/Engine/Source/Editor/TextureEditor/Private/Menus/TextureEditorViewOptionsMenu.h
--- a/Engine/Source/Editor/TextureEditor/Private/Menus/TextureEditorViewOptionsMenu.h
+++ b/Engine/Source/Editor/TextureEditor/Private/Menus/TextureEditorViewOptionsMenu.h
@@ -22,20 +22,11 @@ public:
 	 */
 	static void MakeMenu( FMenuBuilder& MenuBuilder, bool bIsVolumeTexture )
 	{
-		// color channel options
-		MenuBuilder.BeginSection("ChannelSection", LOCTEXT("ChannelsSectionHeader", "Color Channels"));
-		{
-			MenuBuilder.AddMenuEntry(FTextureEditorCommands::Get().RedChannel);
-			MenuBuilder.AddMenuEntry(FTextureEditorCommands::Get().GreenChannel);
-			MenuBuilder.AddMenuEntry(FTextureEditorCommands::Get().BlueChannel);
-			MenuBuilder.AddMenuEntry(FTextureEditorCommands::Get().AlphaChannel);
-			MenuBuilder.AddMenuEntry(FTextureEditorCommands::Get().Desaturation);
-		}
-		MenuBuilder.EndSection();
-
 		// view port options
 		MenuBuilder.BeginSection("ViewportSection", LOCTEXT("ViewportSectionHeader", "Viewport Options"));
 		{
+			MenuBuilder.AddMenuEntry(FTextureEditorCommands::Get().Desaturation);
+
 			MenuBuilder.AddSubMenu(
 				LOCTEXT("Background", "Background"),
 				LOCTEXT("BackgroundTooltip", "Set the viewport's background"),
diff --git a/Engine/Source/Editor/TextureEditor/Private/Models/TextureEditorCommands.cpp b/Engine/Source/Editor/TextureEditor/Private/Models/TextureEditorCommands.cpp
--- a/Engine/Source/Editor/TextureEditor/Private/Models/TextureEditorCommands.cpp
+++ b/Engine/Source/Editor/TextureEditor/Private/Models/TextureEditorCommands.cpp
@@ -25,7 +25,7 @@ void FTextureEditorCommands::RegisterCommands()
 
 	UI_COMMAND(CompressNow, "Compress", "Compress the texture", EUserInterfaceActionType::Button, FInputChord());
 	UI_COMMAND(Reimport, "Reimport", "Reimports the texture from file", EUserInterfaceActionType::Button, FInputChord());
-	UI_COMMAND(Settings, "Settings...", "Opens the settings for the texture editor", EUserInterfaceActionType::Button, FInputChord());
+	UI_COMMAND(Settings, "Advanced Settings...", "Opens the advanced settings for the texture editor", EUserInterfaceActionType::Button, FInputChord());
 }
 
 #undef LOCTEXT_NAMESPACE
diff --git a/Engine/Source/Editor/TextureEditor/Private/Models/TextureEditorViewportClient.cpp b/Engine/Source/Editor/TextureEditor/Private/Models/TextureEditorViewportClient.cpp
--- a/Engine/Source/Editor/TextureEditor/Private/Models/TextureEditorViewportClient.cpp
+++ b/Engine/Source/Editor/TextureEditor/Private/Models/TextureEditorViewportClient.cpp
@@ -61,8 +61,10 @@ void FTextureEditorViewportClient::Draw(FViewport* Viewport, FCanvas* Canvas)
 	{
 		return;
 	}
+
+	TSharedPtr<ITextureEditorToolkit> TextureEditorPinned = TextureEditorPtr.Pin();
 	
-	UTexture* Texture = TextureEditorPtr.Pin()->GetTexture();
+	UTexture* Texture = TextureEditorPinned->GetTexture();
 	FVector2D Ratio = FVector2D(GetViewportHorizontalScrollBarRatio(), GetViewportVerticalScrollBarRatio());
 	FVector2D ViewportSize = FVector2D(TextureEditorViewportPtr.Pin()->GetViewport()->GetSizeXY().X, TextureEditorViewportPtr.Pin()->GetViewport()->GetSizeXY().Y);
 	FVector2D ScrollBarPos = GetViewportScrollBarPositions();
@@ -93,13 +95,13 @@ void FTextureEditorViewportClient::Draw(FViewport* Viewport, FCanvas* Canvas)
 		Texture2D->WaitForStreaming();
 	}
 
-	TextureEditorPtr.Pin()->PopulateQuickInfo();
+	TextureEditorPinned->PopulateQuickInfo();
 
 	// Figure out the size we need
 	uint32 Width, Height;
-	TextureEditorPtr.Pin()->CalculateTextureDimensions(Width, Height);
-	const float MipLevel = (float)TextureEditorPtr.Pin()->GetMipLevel();
-	const float LayerIndex = (float)TextureEditorPtr.Pin()->GetLayer();
+	TextureEditorPinned->CalculateTextureDimensions(Width, Height);
+	const float MipLevel = (float)TextureEditorPinned->GetMipLevel();
+	const float LayerIndex = (float)TextureEditorPinned->GetLayer();
 
 	bool bIsVirtualTexture = false;
 
@@ -117,9 +119,9 @@ void FTextureEditorViewportClient::Draw(FViewport* Viewport, FCanvas* Canvas)
 				Settings.VolumeViewMode == TextureEditorVolumeViewMode_DepthSlices, 
 				FMath::Max<int32>(VolumeTexture->GetSizeZ(), 1), 
 				MipLevel, 
-				(float)TextureEditorPtr.Pin()->GetVolumeOpacity(),
+				(float)TextureEditorPinned->GetVolumeOpacity(),
 				true, 
-				TextureEditorPtr.Pin()->GetVolumeOrientation());
+				TextureEditorPinned->GetVolumeOrientation());
 		}
 		else if (RTTextureVolume)
 		{
@@ -127,9 +129,9 @@ void FTextureEditorViewportClient::Draw(FViewport* Viewport, FCanvas* Canvas)
 				Settings.VolumeViewMode == TextureEditorVolumeViewMode_DepthSlices,
 				FMath::Max<int32>(RTTextureVolume->SizeZ >> RTTextureVolume->GetCachedLODBias(), 1),
 				MipLevel,
-				(float)TextureEditorPtr.Pin()->GetVolumeOpacity(),
+				(float)TextureEditorPinned->GetVolumeOpacity(),
 				true,
-				TextureEditorPtr.Pin()->GetVolumeOrientation());
+				TextureEditorPinned->GetVolumeOrientation());
 		}
 		else if (Texture2D)
 		{
@@ -176,12 +178,12 @@ void FTextureEditorViewportClient::Draw(FViewport* Viewport, FCanvas* Canvas)
 		}
 	}
 
-	float Exposure = FMath::Pow(2.0f, (float)TextureEditorViewportPtr.Pin()->GetExposureBias());
+	float Exposure = FMath::Pow(2.0f, (float)TextureEditorPinned->GetExposureBias());
 
 	if ( Texture->Resource != nullptr )
 	{
 		FCanvasTileItem TileItem( FVector2D( XPos, YPos ), Texture->Resource, FVector2D( Width, Height ), FLinearColor(Exposure, Exposure, Exposure) );
-		TileItem.BlendMode = TextureEditorPtr.Pin()->GetColourChannelBlendMode();
+		TileItem.BlendMode = TextureEditorPinned->GetColourChannelBlendMode();
 		TileItem.BatchedElementParameters = BatchedElementParameters;
 
 		if (bIsVirtualTexture && Texture->Source.GetNumBlocks() > 1)
@@ -213,7 +215,7 @@ void FTextureEditorViewportClient::Draw(FViewport* Viewport, FCanvas* Canvas)
  			auto ZeroOrAbsolute = [](int32 value) -> int32 { return value >= 0 ? 0 : FMath::Abs(value); };
 			auto GetCorrectDimension = [](int value, int32 viewportSize, int32 TextureSize) -> int32 { return value <= viewportSize ? TextureSize : viewportSize; };
 
-			const float Zoom = 1.0f / TextureEditorPtr.Pin()->GetCustomZoomLevel();
+			const float Zoom = 1.0f / TextureEditorPinned->GetCustomZoomLevel();
 			const int32 VisibleXPos = FMath::FloorToInt(Zoom * -FMath::Min(0, XPos));
 			const int32 VisibleYPos = FMath::FloorToInt(Zoom * -FMath::Min(0, YPos));
 			
diff --git a/Engine/Source/Editor/TextureEditor/Private/TextureEditorClasses.cpp b/Engine/Source/Editor/TextureEditor/Private/TextureEditorClasses.cpp
--- a/Engine/Source/Editor/TextureEditor/Private/TextureEditorClasses.cpp
+++ b/Engine/Source/Editor/TextureEditor/Private/TextureEditorClasses.cpp
@@ -7,10 +7,10 @@
 UTextureEditorSettings::UTextureEditorSettings( const FObjectInitializer& ObjectInitializer )
 	: Super(ObjectInitializer)
 	, Background(TextureEditorBackground_Checkered)
-	, BackgroundColor(FColor::Black)
+	, BackgroundColor(FColor::FromHex("242424FF"))
 	, CheckerColorOne(FColor(128, 128, 128))
 	, CheckerColorTwo(FColor(64, 64, 64))
-	, CheckerSize(32)
+	, CheckerSize(16)
 	, FitToViewport(true)
 	, ZoomMode(ETextureEditorZoomMode::Fit)
 	, TextureBorderColor(FColor::White)
diff --git a/Engine/Source/Editor/TextureEditor/Private/TextureEditorConstants.h b/Engine/Source/Editor/TextureEditor/Private/TextureEditorConstants.h
--- a/Engine/Source/Editor/TextureEditor/Private/TextureEditorConstants.h
+++ b/Engine/Source/Editor/TextureEditor/Private/TextureEditorConstants.h
@@ -7,3 +7,9 @@
 const double MaxZoom = 16.0;
 const double MinZoom = 0.01;
 const double ZoomStep = 0.1;
+
+// Specifies the maximum allowed exposure bias.
+const int32 MaxExposure = 10;
+
+// Specifies the minimum allowed exposure bias.
+const int32 MinExposure = -10;
diff --git a/Engine/Source/Editor/TextureEditor/Private/TextureEditorToolkit.cpp b/Engine/Source/Editor/TextureEditor/Private/TextureEditorToolkit.cpp
--- a/Engine/Source/Editor/TextureEditor/Private/TextureEditorToolkit.cpp
+++ b/Engine/Source/Editor/TextureEditor/Private/TextureEditorToolkit.cpp
@@ -41,6 +41,9 @@
 #include "Curves/CurveLinearColorAtlas.h"
 #include "TextureEditorSettings.h"
 #include "TextureCompiler.h"
+#include "Widgets/Input/SSlider.h"
+#include "Widgets/Layout/SSpacer.h"
+#include "Menus/TextureEditorViewOptionsMenu.h"
 
 #define LOCTEXT_NAMESPACE "FTextureEditorToolkit"
 
@@ -145,6 +148,10 @@ void FTextureEditorToolkit::InitTextureEditor( const EToolkitMode::Type Mode, co
 	bIsBlueChannel = true;
 	bIsAlphaChannel = false;
 
+	ExposureBias = 0;
+
+	bIsVolumeTexture = Texture->IsA<UVolumeTexture>() || Texture->IsA<UTextureRenderTargetVolume>();
+
 	switch (Texture->CompressionSettings)
 	{
 	default:
@@ -179,7 +186,7 @@ void FTextureEditorToolkit::InitTextureEditor( const EToolkitMode::Type Mode, co
 	BindCommands();
 	CreateInternalWidgets();
 
-	const TSharedRef<FTabManager::FLayout> StandaloneDefaultLayout = FTabManager::NewLayout("Standalone_TextureEditor_Layout_v3")
+	const TSharedRef<FTabManager::FLayout> StandaloneDefaultLayout = FTabManager::NewLayout("Standalone_TextureEditor_Layout_v4")
 		->AddArea
 		(
 			FTabManager::NewPrimaryArea()
@@ -188,15 +195,6 @@ void FTextureEditorToolkit::InitTextureEditor( const EToolkitMode::Type Mode, co
 				(
 					FTabManager::NewSplitter()
 						->SetOrientation(Orient_Vertical)
-						->SetSizeCoefficient(0.66f)
-						->Split
-						(
-							FTabManager::NewStack()
-								->AddTab(GetToolbarTabId(), ETabState::OpenedTab)
-								->SetHideTabWell(true)
-								->SetSizeCoefficient(0.1f)
-								
-						)
 						->Split
 						(
 							FTabManager::NewStack()
@@ -790,27 +788,23 @@ void FTextureEditorToolkit::BindCommands( )
 
 	ToolkitCommands->MapAction(
 		Commands.RedChannel,
-		FExecuteAction::CreateSP(this, &FTextureEditorToolkit::HandleRedChannelActionExecute),
-		FCanExecuteAction(),
-		FIsActionChecked::CreateSP(this, &FTextureEditorToolkit::HandleRedChannelActionIsChecked));
+		FExecuteAction::CreateSP(this, &FTextureEditorToolkit::OnChannelButtonCheckStateChanged, ETextureChannelButton::Red),
+		FCanExecuteAction());
 
 	ToolkitCommands->MapAction(
 		Commands.GreenChannel,
-		FExecuteAction::CreateSP(this, &FTextureEditorToolkit::HandleGreenChannelActionExecute),
-		FCanExecuteAction(),
-		FIsActionChecked::CreateSP(this, &FTextureEditorToolkit::HandleGreenChannelActionIsChecked));
+		FExecuteAction::CreateSP(this, &FTextureEditorToolkit::OnChannelButtonCheckStateChanged, ETextureChannelButton::Green),
+		FCanExecuteAction());
 
 	ToolkitCommands->MapAction(
 		Commands.BlueChannel,
-		FExecuteAction::CreateSP(this, &FTextureEditorToolkit::HandleBlueChannelActionExecute),
-		FCanExecuteAction(),
-		FIsActionChecked::CreateSP(this, &FTextureEditorToolkit::HandleBlueChannelActionIsChecked));
+		FExecuteAction::CreateSP(this, &FTextureEditorToolkit::OnChannelButtonCheckStateChanged, ETextureChannelButton::Blue),
+		FCanExecuteAction());
 
 	ToolkitCommands->MapAction(
 		Commands.AlphaChannel,
-		FExecuteAction::CreateSP(this, &FTextureEditorToolkit::HandleAlphaChannelActionExecute),
-		FCanExecuteAction(),
-		FIsActionChecked::CreateSP(this, &FTextureEditorToolkit::HandleAlphaChannelActionIsChecked));
+		FExecuteAction::CreateSP(this, &FTextureEditorToolkit::OnChannelButtonCheckStateChanged, ETextureChannelButton::Alpha),
+		FCanExecuteAction());
 
 	ToolkitCommands->MapAction(
 		Commands.Desaturation,
@@ -907,92 +901,88 @@ void FTextureEditorToolkit::CreateInternalWidgets( )
 	.AutoHeight()
 	.Padding(2.0f)
 	[
-		SNew(SBorder)
+		SNew(SHorizontalBox)
+		+ SHorizontalBox::Slot()
+		.FillWidth(0.5f)
 		[
-			SNew(SHorizontalBox)
+			SNew(SVerticalBox)
 
-			+ SHorizontalBox::Slot()
-			.FillWidth(0.5f)
+			+ SVerticalBox::Slot()
+			.AutoHeight()
+			.VAlign(VAlign_Center)
+			.Padding(4.0f)
 			[
-				SNew(SVerticalBox)
-
-				+ SVerticalBox::Slot()
-				.AutoHeight()
-				.VAlign(VAlign_Center)
-				.Padding(4.0f)
-				[
-					SAssignNew(ImportedText, STextBlock)
-				]
-
-				+ SVerticalBox::Slot()
-				.AutoHeight()
-				.VAlign(VAlign_Center)
-				.Padding(4.0f)
-				[
-					SAssignNew(CurrentText, STextBlock)
-				]
+				SAssignNew(ImportedText, STextBlock)
+			]
 
-				+ SVerticalBox::Slot()
-				.AutoHeight()
-				.VAlign(VAlign_Center)
-				.Padding(4.0f)
-				[
-					SAssignNew(MaxInGameText, STextBlock)
-				]
+			+ SVerticalBox::Slot()
+			.AutoHeight()
+			.VAlign(VAlign_Center)
+			.Padding(4.0f)
+			[
+				SAssignNew(CurrentText, STextBlock)
+			]
 
-				+ SVerticalBox::Slot()
-				.AutoHeight()
-				.VAlign(VAlign_Center)
-				.Padding(4.0f)
-				[
-					SAssignNew(SizeText, STextBlock)
-				]
+			+ SVerticalBox::Slot()
+			.AutoHeight()
+			.VAlign(VAlign_Center)
+			.Padding(4.0f)
+			[
+				SAssignNew(MaxInGameText, STextBlock)
+			]
 
-				+ SVerticalBox::Slot()
-				.AutoHeight()
-				.VAlign(VAlign_Center)
-				.Padding(4.0f)
-				[
-					SAssignNew(HasAlphaChannelText, STextBlock)
-				]
+			+ SVerticalBox::Slot()
+			.AutoHeight()
+			.VAlign(VAlign_Center)
+			.Padding(4.0f)
+			[
+				SAssignNew(SizeText, STextBlock)
 			]
 
-			+ SHorizontalBox::Slot()
-			.FillWidth(0.5f)
+			+ SVerticalBox::Slot()
+			.AutoHeight()
+			.VAlign(VAlign_Center)
+			.Padding(4.0f)
 			[
-				SNew(SVerticalBox)
+				SAssignNew(HasAlphaChannelText, STextBlock)
+			]
+		]
 
-				+ SVerticalBox::Slot()
-				.AutoHeight()
-				.VAlign(VAlign_Center)
-				.Padding(4.0f)
-				[
-					SAssignNew(MethodText, STextBlock)
-				]
+		+ SHorizontalBox::Slot()
+		.FillWidth(0.5f)
+		[
+			SNew(SVerticalBox)
 
-				+ SVerticalBox::Slot()
-				.AutoHeight()
-				.VAlign(VAlign_Center)
-				.Padding(4.0f)
-				[
-					SAssignNew(FormatText, STextBlock)
-				]
+			+ SVerticalBox::Slot()
+			.AutoHeight()
+			.VAlign(VAlign_Center)
+			.Padding(4.0f)
+			[
+				SAssignNew(MethodText, STextBlock)
+			]
 
-				+ SVerticalBox::Slot()
-				.AutoHeight()
-				.VAlign(VAlign_Center)
-				.Padding(4.0f)
-				[
-					SAssignNew(LODBiasText, STextBlock)
-				]
+			+ SVerticalBox::Slot()
+			.AutoHeight()
+			.VAlign(VAlign_Center)
+			.Padding(4.0f)
+			[
+				SAssignNew(FormatText, STextBlock)
+			]
 
-				+ SVerticalBox::Slot()
-				.AutoHeight()
-				.VAlign(VAlign_Center)
-				.Padding(4.0f)
-				[
-					SAssignNew(NumMipsText, STextBlock)
-				]
+			+ SVerticalBox::Slot()
+			.AutoHeight()
+			.VAlign(VAlign_Center)
+			.Padding(4.0f)
+			[
+				SAssignNew(LODBiasText, STextBlock)
+			]
+
+			+ SVerticalBox::Slot()
+			.AutoHeight()
+			.VAlign(VAlign_Center)
+			.Padding(4.0f)
+			[
+				SAssignNew(NumMipsText, STextBlock)
 			]
 		]
 	]
@@ -1001,166 +991,19 @@ void FTextureEditorToolkit::CreateInternalWidgets( )
 	.FillHeight(1.0f)
 	.Padding(2.0f)
 	[
-		SNew(SBorder)
-		.Padding(4.0f)
-		[
-			BuildTexturePropertiesWidget()
-		]
+		BuildTexturePropertiesWidget()
 	];
 }
 
-
-BEGIN_SLATE_FUNCTION_BUILD_OPTIMIZATION
-void FTextureEditorToolkit::ExtendToolBar( )
+void FTextureEditorToolkit::ExtendToolBar()
 {
-	TSharedRef<SWidget> LODControl = SNew(SBox)
-		.WidthOverride(240.0f)
-		[
-			SNew(SHorizontalBox)
-
-			+ SHorizontalBox::Slot()
-				.FillWidth(1.0f)
-				.MaxWidth(240.0f)
-				.Padding(0.0f, 0.0f, 0.0f, 0.0f)
-				.VAlign(VAlign_Center)
-				[
-					// Mip and exposure controls
-					SNew(SHorizontalBox)
-
-					+ SHorizontalBox::Slot()
-						.Padding(4.0f, 0.0f, 4.0f, 0.0f)
-						.AutoWidth()
-						[
-							SNew(SHorizontalBox)
-
-							+ SHorizontalBox::Slot()
-								.VAlign(VAlign_Center)
-								.AutoWidth()
-								[
-									SNew(SCheckBox)
-										.IsChecked(this, &FTextureEditorToolkit::HandleMipLevelCheckBoxIsChecked)
-										.IsEnabled(this, &FTextureEditorToolkit::HandleMipLevelCheckBoxIsEnabled)
-										.OnCheckStateChanged(this, &FTextureEditorToolkit::HandleMipLevelCheckBoxCheckedStateChanged)
-								]
-						]
-
-					+ SHorizontalBox::Slot()
-						.Padding(4.0f, 0.0f, 4.0f, 0.0f)
-						.FillWidth(1)
-						[
-							SNew(SHorizontalBox)
-
-							+ SHorizontalBox::Slot()
-								.Padding(0.0f, 0.0, 4.0, 0.0)
-								.AutoWidth()
-								.VAlign(VAlign_Center)
-								[
-									SNew(STextBlock)
-										.Text(NSLOCTEXT("TextureEditor", "MipLevel", "Mip Level: "))
-								]
-
-							+ SHorizontalBox::Slot()
-								.VAlign(VAlign_Center)
-								.FillWidth(1.0f)
-								[
-									SNew(SNumericEntryBox<int32>)
-										.AllowSpin(true)
-										.MinSliderValue(MIPLEVEL_MIN)
-										.MaxSliderValue(this, &FTextureEditorToolkit::GetMaxMipLevel)
-										.Value(this, &FTextureEditorToolkit::HandleMipLevelEntryBoxValue)
-										.OnValueChanged(this, &FTextureEditorToolkit::HandleMipLevelEntryBoxChanged)
-										.IsEnabled(this, &FTextureEditorToolkit::GetUseSpecifiedMip)
-								]
-
-							+ SHorizontalBox::Slot()
-								.AutoWidth()
-								.VAlign(VAlign_Center)
-								.Padding(2.0f)
-								[
-									SNew(SButton)
-										.Text(NSLOCTEXT("TextureEditor", "MipMinus", "-"))
-										.OnClicked(this, &FTextureEditorToolkit::HandleMipMapMinusButtonClicked)
-										.IsEnabled(this, &FTextureEditorToolkit::GetUseSpecifiedMip)
-								]
-
-							+ SHorizontalBox::Slot()
-								.AutoWidth()
-								.VAlign(VAlign_Center)
-								.Padding(2.0f)
-								[
-									SNew(SButton)
-										.Text(NSLOCTEXT("TextureEditor", "MipPlus", "+"))
-										.OnClicked(this, &FTextureEditorToolkit::HandleMipMapPlusButtonClicked)
-										.IsEnabled(this, &FTextureEditorToolkit::GetUseSpecifiedMip)
-								]
-						]
-				]
-			];
-
-	TSharedRef<SWidget> LayerControl = SNew(SBox)
-		.WidthOverride(200.0f)
-		[
-			SNew(SHorizontalBox)
-
-				+SHorizontalBox::Slot()
-					.FillWidth(1.0f)
-					.MaxWidth(200.0f)
-					.Padding(0.0f, 0.0f, 0.0f, 0.0f)
-					.VAlign(VAlign_Center)
-					[
-						// Layer controls
-						SNew(SHorizontalBox)
-
-						+ SHorizontalBox::Slot()
-							.Padding(0.0f, 0.0, 4.0, 0.0)
-							.AutoWidth()
-							.VAlign(VAlign_Center)
-							[
-								SNew(STextBlock)
-									.Text(NSLOCTEXT("TextureEditor", "Layer", "Layer: "))
-							]
-
-						+ SHorizontalBox::Slot()
-							.VAlign(VAlign_Center)
-							.FillWidth(1.0f)
-							[
-								SNew(SNumericEntryBox<int32>)
-									.AllowSpin(true)
-									.MinSliderValue(MIPLEVEL_MIN)
-									.MaxSliderValue(this, &FTextureEditorToolkit::GetMaxLayer)
-									.Value(this, &FTextureEditorToolkit::HandleLayerEntryBoxValue)
-									.OnValueChanged(this, &FTextureEditorToolkit::HandleLayerEntryBoxChanged)
-							]
-
-						+ SHorizontalBox::Slot()
-							.AutoWidth()
-							.VAlign(VAlign_Center)
-							.Padding(2.0f)
-							[
-								SNew(SButton)
-									.Text(NSLOCTEXT("TextureEditor", "LayerMinus", "-"))
-									.OnClicked(this, &FTextureEditorToolkit::HandleLayerMinusButtonClicked)
-							]
-
-						+ SHorizontalBox::Slot()
-							.AutoWidth()
-							.VAlign(VAlign_Center)
-							.Padding(2.0f)
-							[
-								SNew(SButton)
-									.Text(NSLOCTEXT("TextureEditor", "LayerPlus", "+"))
-									.OnClicked(this, &FTextureEditorToolkit::HandleLayerPlusButtonClicked)
-							]
-					]
-		];
-
 	TSharedPtr<FExtender> ToolbarExtender = MakeShareable(new FExtender);
 
 	ToolbarExtender->AddToolBarExtension(
 		"Asset",
 		EExtensionHook::After,
 		GetToolkitCommands(),
-		FToolBarExtensionDelegate::CreateSP(this, &FTextureEditorToolkit::FillToolbar, GetToolkitCommands(), LODControl, LayerControl)
+		FToolBarExtensionDelegate::CreateSP(this, &FTextureEditorToolkit::FillToolbar)
 	);
 
 	AddToolbarExtender(ToolbarExtender);
@@ -1169,8 +1012,15 @@ void FTextureEditorToolkit::ExtendToolBar( )
 	AddToolbarExtender(TextureEditorModule->GetToolBarExtensibilityManager()->GetAllExtenders(GetToolkitCommands(), GetEditingObjects()));
 }
 
-void FTextureEditorToolkit::FillToolbar(FToolBarBuilder& ToolbarBuilder, const TSharedRef< FUICommandList > InToolkitCommands, TSharedRef<SWidget> LODControl, TSharedRef<SWidget> LayerControl)
+void FTextureEditorToolkit::FillToolbar(FToolBarBuilder& ToolbarBuilder)
 {
+	TSharedRef<SWidget> ChannelControl = MakeChannelControlWidget();
+	TSharedRef<SWidget> LODControl = MakeLODControlWidget();
+	TSharedRef<SWidget> LayerControl = MakeLayerControlWidget();
+	TSharedRef<SWidget> ExposureControl = MakeExposureContolWidget();
+	TSharedPtr<SWidget> OptionalOpacityControl = IsVolumeTexture() ? TSharedPtr<SWidget>(MakeOpacityControlWidget()) : nullptr;
+	TSharedRef<SWidget> ZoomControl = MakeZoomControlWidget();
+
 	UCurveLinearColorAtlas* Atlas = Cast<UCurveLinearColorAtlas>(GetTexture());
 	if (!Atlas)
 	{
@@ -1181,10 +1031,17 @@ void FTextureEditorToolkit::FillToolbar(FToolBarBuilder& ToolbarBuilder, const T
 		}
 		ToolbarBuilder.EndSection();
 
+		ToolbarBuilder.BeginSection("Channels");
+		{
+			ToolbarBuilder.AddWidget(ChannelControl);
+		}
+		ToolbarBuilder.EndSection();
+
 		ToolbarBuilder.BeginSection("TextureMipAndExposure");
 		{
 			ToolbarBuilder.AddWidget(LODControl);
-		}
+			ToolbarBuilder.AddWidget(ExposureControl);
+		}	
 		ToolbarBuilder.EndSection();
 
 		if (HasLayers())
@@ -1195,11 +1052,34 @@ void FTextureEditorToolkit::FillToolbar(FToolBarBuilder& ToolbarBuilder, const T
 			}
 			ToolbarBuilder.EndSection();
 		}
-	}
-}
 
-END_SLATE_FUNCTION_BUILD_OPTIMIZATION
+		if (OptionalOpacityControl.IsValid())
+		{
+			ToolbarBuilder.BeginSection("Opacity");
+			{
+				ToolbarBuilder.AddWidget(OptionalOpacityControl.ToSharedRef());
+			}
+			ToolbarBuilder.EndSection();
+		}
 
+		ToolbarBuilder.BeginSection("Zoom");
+		{
+			ToolbarBuilder.AddWidget(ZoomControl);
+		}
+		ToolbarBuilder.EndSection();
+		ToolbarBuilder.BeginSection("Settings");
+		{
+			ToolbarBuilder.AddWidget(SNew(SSpacer), NAME_None, false, HAlign_Right);
+			ToolbarBuilder.AddComboButton(
+				FUIAction(),
+				FOnGetContent::CreateSP(this, &FTextureEditorToolkit::OnGenerateSettingsMenu),
+				LOCTEXT("SettingsMenu", "Settings"),
+				FText::GetEmpty(),
+				FSlateIcon(FAppStyle::GetAppStyleSetName(), "Icons.Settings")
+			);
+		}
+	}
+}
 
 TOptional<int32> FTextureEditorToolkit::GetMaxMipLevel( ) const
 {
@@ -1266,10 +1146,42 @@ bool FTextureEditorToolkit::IsCubeTexture( ) const
 }
 
 
+TSharedRef<SWidget> FTextureEditorToolkit::OnGenerateMipMapLevelMenu()
+{
+	FMenuBuilder MenuBuilder(true, nullptr);
+
+	for (int32 MipLevel = MIPLEVEL_MIN; MipLevel <= GetMaxMipLevel().Get(MIPLEVEL_MAX); ++MipLevel)
+	{
+		FText MipNumberText = FText::AsNumber(MipLevel);
+
+		MenuBuilder.AddMenuEntry(
+			FText::Format(LOCTEXT("MipLevel", "Mip Level {0}"), MipNumberText),
+			FText::Format(LOCTEXT("MipLevel_Tooltip", "Display Mip Level {0}"), MipNumberText),
+			FSlateIcon(),
+			FUIAction(
+				FExecuteAction::CreateSP(this, &FTextureEditorToolkit::HandleMipLevelChanged, MipLevel),
+				FCanExecuteAction(),
+				FIsActionChecked::CreateLambda([this, MipLevel]() {return SpecifiedMipLevel == MipLevel; })
+			)
+		);
+	}
+
+	return MenuBuilder.MakeWidget();
+}
+
+TSharedRef<SWidget> FTextureEditorToolkit::OnGenerateSettingsMenu()
+{
+	FMenuBuilder MenuBuilder(true, ToolkitCommands);
+	FTextureEditorViewOptionsMenu::MakeMenu(MenuBuilder, IsVolumeTexture());
+
+	return MenuBuilder.MakeWidget();
+}
+
 /* FTextureEditorToolkit callbacks
  *****************************************************************************/
 
-bool FTextureEditorToolkit::HandleAlphaChannelActionCanExecute( ) const
+
+bool FTextureEditorToolkit::IsAlphaChannelButtonEnabled( ) const
 {
 	const UTexture2D* Texture2D = Cast<UTexture2D>(Texture);
 
@@ -1281,28 +1193,89 @@ bool FTextureEditorToolkit::HandleAlphaChannelActionCanExecute( ) const
 	return Texture2D->HasAlphaChannel();
 }
 
-
-void FTextureEditorToolkit::HandleAlphaChannelActionExecute( )
+FSlateColor FTextureEditorToolkit::GetChannelButtonBackgroundColor(ETextureChannelButton Button) const
 {
-	bIsAlphaChannel = !bIsAlphaChannel;
-}
+	FSlateColor Dropdown = FAppStyle::Get().GetSlateColor("Colors.Dropdown");
 
+	switch (Button)
+	{
+	case ETextureChannelButton::Red:
+		return bIsRedChannel ? FLinearColor::Red : FLinearColor::White;
+	case ETextureChannelButton::Green:
+		return bIsGreenChannel ? FLinearColor::Green : FLinearColor::White;
+	case ETextureChannelButton::Blue:
+		return bIsBlueChannel ? FLinearColor::Blue : FLinearColor::White;
+	case ETextureChannelButton::Alpha:
+		return bIsAlphaChannel ? FLinearColor::White : FLinearColor::White;
+	default:
+		check(false);
+		return FSlateColor();
+	}
+}
 
-bool FTextureEditorToolkit::HandleAlphaChannelActionIsChecked( ) const
+FSlateColor FTextureEditorToolkit::GetChannelButtonForegroundColor(ETextureChannelButton Button) const
 {
-	return bIsAlphaChannel;
-}
+	FSlateColor DefaultForeground = FAppStyle::Get().GetSlateColor("Colors.Foreground");
 
+	switch (Button)
+	{
+	case ETextureChannelButton::Red:
+		return bIsRedChannel ? FLinearColor::Black : DefaultForeground;
+	case ETextureChannelButton::Green:
+		return bIsGreenChannel ? FLinearColor::Black : DefaultForeground;
+	case ETextureChannelButton::Blue:
+		return bIsBlueChannel ? FLinearColor::Black : DefaultForeground;
+	case ETextureChannelButton::Alpha:
+		return bIsAlphaChannel ? FLinearColor::Black : DefaultForeground;
+	default:
+		check(false);
+		return FSlateColor::UseForeground();
+	}
+}
 
-void FTextureEditorToolkit::HandleBlueChannelActionExecute( )
+void FTextureEditorToolkit::OnChannelButtonCheckStateChanged(ETextureChannelButton Button)
 {
-	 bIsBlueChannel = !bIsBlueChannel;
+	switch (Button)
+	{
+	case ETextureChannelButton::Red:
+		bIsRedChannel = !bIsRedChannel;
+		break;
+	case ETextureChannelButton::Green:
+		bIsGreenChannel = !bIsGreenChannel;
+		break;
+	case ETextureChannelButton::Blue:
+		bIsBlueChannel = !bIsBlueChannel;
+		break;
+	case ETextureChannelButton::Alpha:
+		bIsAlphaChannel = !bIsAlphaChannel;
+		break;
+	default:
+		check(false);
+		break;
+	}
 }
 
-
-bool FTextureEditorToolkit::HandleBlueChannelActionIsChecked( ) const
+ECheckBoxState FTextureEditorToolkit::OnGetChannelButtonCheckState(ETextureChannelButton Button) const
 {
-	return bIsBlueChannel;
+	switch (Button)
+	{
+	case ETextureChannelButton::Red:
+		return bIsRedChannel ? ECheckBoxState::Checked : ECheckBoxState::Unchecked;
+		break;
+	case ETextureChannelButton::Green:
+		return bIsGreenChannel ? ECheckBoxState::Checked : ECheckBoxState::Unchecked;
+		break;
+	case ETextureChannelButton::Blue:
+		return bIsBlueChannel ? ECheckBoxState::Checked : ECheckBoxState::Unchecked;
+		break;
+	case ETextureChannelButton::Alpha:
+		return bIsAlphaChannel ? ECheckBoxState::Checked : ECheckBoxState::Unchecked;
+		break;
+	default:
+		check(false);
+		return ECheckBoxState::Unchecked;
+		break;
+	}
 }
 
 
@@ -1378,18 +1351,6 @@ void FTextureEditorToolkit::HandleZoomToNaturalActionExecute()
 	SetCustomZoomLevel(1);
 }
 
-void FTextureEditorToolkit::HandleGreenChannelActionExecute( )
-{
-	 bIsGreenChannel = !bIsGreenChannel;
-}
-
-
-bool FTextureEditorToolkit::HandleGreenChannelActionIsChecked( ) const
-{
-	return bIsGreenChannel;
-}
-
-
 void FTextureEditorToolkit::HandleMipLevelCheckBoxCheckedStateChanged( ECheckBoxState InNewState )
 {
 	bUseSpecifiedMipLevel = InNewState == ECheckBoxState::Checked;
@@ -1414,36 +1375,28 @@ bool FTextureEditorToolkit::HandleMipLevelCheckBoxIsEnabled( ) const
 	return true;
 }
 
-
-void FTextureEditorToolkit::HandleMipLevelEntryBoxChanged( int32 NewMipLevel )
+void FTextureEditorToolkit::HandleMipLevelChanged(int32 NewMipLevel)
 {
 	SpecifiedMipLevel = FMath::Clamp<int32>(NewMipLevel, MIPLEVEL_MIN, GetMaxMipLevel().Get(MIPLEVEL_MAX));
-}
 
+	MipLevelTextBlock->SetText(FText::Format(LOCTEXT("MipLevel", "Mip Level {0}"), SpecifiedMipLevel));
+}
 
 TOptional<int32> FTextureEditorToolkit::HandleMipLevelEntryBoxValue( ) const
 {
 	return SpecifiedMipLevel;
 }
 
-
-FReply FTextureEditorToolkit::HandleMipMapMinusButtonClicked( )
+FReply FTextureEditorToolkit::HandleMipMapMinusButtonClicked()
 {
-	if (SpecifiedMipLevel > MIPLEVEL_MIN)
-	{
-		--SpecifiedMipLevel;
-	}
+	HandleMipLevelChanged(--SpecifiedMipLevel);
 
 	return FReply::Handled();
 }
 
-
-FReply FTextureEditorToolkit::HandleMipMapPlusButtonClicked( )
+FReply FTextureEditorToolkit::HandleMipMapPlusButtonClicked()
 {
-	if (SpecifiedMipLevel < GetMaxMipLevel().Get(MIPLEVEL_MAX))
-	{
-		++SpecifiedMipLevel;
-	}
+	HandleMipLevelChanged(++SpecifiedMipLevel);
 
 	return FReply::Handled();
 }
@@ -1454,28 +1407,6 @@ void FTextureEditorToolkit::HandleLayerEntryBoxChanged(int32 NewLayer)
 	PopulateQuickInfo();
 }
 
-FReply FTextureEditorToolkit::HandleLayerMinusButtonClicked()
-{
-	if (SpecifiedLayer > 0)
-	{
-		--SpecifiedLayer;
-		PopulateQuickInfo();
-	}
-
-	return FReply::Handled();
-}
-
-
-FReply FTextureEditorToolkit::HandleLayerPlusButtonClicked()
-{
-	if ((SpecifiedLayer + 1) < Texture->Source.GetNumLayers())
-	{
-		++SpecifiedLayer;
-		PopulateQuickInfo();
-	}
-
-	return FReply::Handled();
-}
 
 TOptional<int32> FTextureEditorToolkit::HandleLayerEntryBoxValue() const
 {
@@ -1487,18 +1418,6 @@ bool FTextureEditorToolkit::HasLayers() const
 	return Texture->Source.GetNumLayers() > 1;
 }
 
-void FTextureEditorToolkit::HandleRedChannelActionExecute( )
-{
-	bIsRedChannel = !bIsRedChannel;
-}
-
-
-bool FTextureEditorToolkit::HandleRedChannelActionIsChecked( ) const
-{
-	return bIsRedChannel;
-}
-
-
 bool FTextureEditorToolkit::HandleReimportActionCanExecute( ) const
 {
 	if (Texture->IsA<ULightMapTexture2D>() || Texture->IsA<UShadowMapTexture2D>() || Texture->IsA<UTexture2DDynamic>() || Texture->IsA<UTextureRenderTarget>() || Texture->IsA<UCurveLinearColorAtlas>())
@@ -1622,5 +1541,437 @@ bool FTextureEditorToolkit::HandleTextureBorderActionIsChecked( ) const
 	return Settings.TextureBorderEnabled;
 }
 
+EVisibility FTextureEditorToolkit::HandleExposureBiasWidgetVisibility() const
+{
+	if ((Texture != nullptr) && (Texture->CompressionSettings == TC_HDR || Texture->CompressionSettings == TC_HDR_Compressed))
+	{
+		return EVisibility::Visible;
+	}
+
+	return EVisibility::Collapsed;
+}
+
+
+TOptional<int32> FTextureEditorToolkit::HandleExposureBiasBoxValue() const
+{
+	return GetExposureBias();
+}
+
+void FTextureEditorToolkit::HandleExposureBiasBoxValueChanged(int32 NewExposure)
+{
+	ExposureBias = NewExposure;
+}
+
+void FTextureEditorToolkit::HandleOpacitySliderChanged(float NewValue)
+{
+	SetVolumeOpacity(NewValue);
+}
+
+TOptional<float> FTextureEditorToolkit::HandleOpacitySliderValue() const
+{
+	return GetVolumeOpacity();
+}
+
+
+FReply FTextureEditorToolkit::HandleViewOptionsMenuButtonClicked()
+{
+	if (ViewOptionsMenuAnchor->ShouldOpenDueToClick())
+	{
+		ViewOptionsMenuAnchor->SetIsOpen(true);
+	}
+	else
+	{
+		ViewOptionsMenuAnchor->SetIsOpen(false);
+	}
+
+	return FReply::Handled();
+}
+
+void FTextureEditorToolkit::HandleZoomMenuEntryClicked(double ZoomValue)
+{
+	SetCustomZoomLevel(ZoomValue);
+}
+
+void FTextureEditorToolkit::HandleZoomMenuFillClicked()
+{
+	SetZoomMode(ETextureEditorZoomMode::Fill);
+}
+
+void FTextureEditorToolkit::HandleZoomMenuFitClicked()
+{
+	SetZoomMode(ETextureEditorZoomMode::Fit);
+}
+
+bool FTextureEditorToolkit::IsZoomMenuFillChecked() const
+{
+	return IsCurrentZoomMode(ETextureEditorZoomMode::Fill);
+}
+
+bool FTextureEditorToolkit::IsZoomMenuFitChecked() const
+{
+	return IsCurrentZoomMode(ETextureEditorZoomMode::Fit);
+}
+
+FText FTextureEditorToolkit::HandleZoomPercentageText() const
+{
+	double DisplayedZoomLevel = CalculateDisplayedZoomLevel();
+	FText ZoomLevelPercent = FText::AsPercent(DisplayedZoomLevel);
+
+	// For fit and fill, show the effective zoom level in parenthesis - eg. "Fill (220%)"
+	static const FText ZoomModeWithPercentFormat = LOCTEXT("ZoomModeWithPercentFormat", "{ZoomMode} ({ZoomPercent})");
+	if (GetZoomMode() == ETextureEditorZoomMode::Fit)
+	{
+		static const FText ZoomModeFit = LOCTEXT("ZoomModeFit", "Fit");
+		return FText::FormatNamed(ZoomModeWithPercentFormat, TEXT("ZoomMode"), ZoomModeFit, TEXT("ZoomPercent"), ZoomLevelPercent);
+	}
+
+	if (GetZoomMode() == ETextureEditorZoomMode::Fill)
+	{
+		static const FText ZoomModeFill = LOCTEXT("ZoomModeFill", "Fill");
+		return FText::FormatNamed(ZoomModeWithPercentFormat, TEXT("ZoomMode"), ZoomModeFill, TEXT("ZoomPercent"), ZoomLevelPercent);
+	}
+
+	// If custom, then just the percent is enough
+	return ZoomLevelPercent;
+}
+
+void FTextureEditorToolkit::HandleZoomSliderChanged(float NewValue)
+{
+	SetCustomZoomLevel(NewValue * MaxZoom);
+}
+
+float FTextureEditorToolkit::HandleZoomSliderValue() const
+{
+	return (CalculateDisplayedZoomLevel() / MaxZoom);
+}
+
+TSharedRef<SWidget> FTextureEditorToolkit::MakeChannelControlWidget()
+{
+	auto OnChannelCheckStateChanged = [this](ECheckBoxState NewState, ETextureChannelButton Button)
+	{
+		OnChannelButtonCheckStateChanged(Button);
+	};
+
+	TSharedRef<SWidget> ChannelControl = 
+		SNew(SHorizontalBox)
+		+SHorizontalBox::Slot()
+		.VAlign(VAlign_Center)
+		.Padding(2.0f)
+		.AutoWidth()
+		[
+			SNew(SCheckBox)
+			.Style(FAppStyle::Get(), "TextureEditor.ChannelButtonStyle")
+			.BorderBackgroundColor(this, &FTextureEditorToolkit::GetChannelButtonBackgroundColor, ETextureChannelButton::Red)
+			.ForegroundColor(this, &FTextureEditorToolkit::GetChannelButtonForegroundColor, ETextureChannelButton::Red)
+			.OnCheckStateChanged_Lambda(OnChannelCheckStateChanged, ETextureChannelButton::Red)
+			.IsChecked(this, &FTextureEditorToolkit::OnGetChannelButtonCheckState, ETextureChannelButton::Red)
+			[
+				SNew(STextBlock)
+				.Font(FAppStyle::Get().GetFontStyle("TextureEditor.ChannelButtonFont"))
+				.Text(FText::FromString("R"))
+			]
+		]
+		+SHorizontalBox::Slot()
+		.VAlign(VAlign_Center)
+		.Padding(2.0f)
+		.AutoWidth()
+		[
+			SNew(SCheckBox)
+			.Style(FAppStyle::Get(), "TextureEditor.ChannelButtonStyle")
+			.BorderBackgroundColor(this, &FTextureEditorToolkit::GetChannelButtonBackgroundColor, ETextureChannelButton::Green)
+			.ForegroundColor(this, &FTextureEditorToolkit::GetChannelButtonForegroundColor, ETextureChannelButton::Green)
+			.OnCheckStateChanged_Lambda(OnChannelCheckStateChanged, ETextureChannelButton::Green)
+			.IsChecked(this, &FTextureEditorToolkit::OnGetChannelButtonCheckState, ETextureChannelButton::Green)
+			[
+				SNew(STextBlock)
+				.Font(FAppStyle::Get().GetFontStyle("TextureEditor.ChannelButtonFont"))
+				.Text(FText::FromString("G"))
+			]
+		]
+
+		+SHorizontalBox::Slot()
+		.VAlign(VAlign_Center)
+		.Padding(2.0f)
+		.AutoWidth()
+		[
+			SNew(SCheckBox)
+			.Style(FAppStyle::Get(), "TextureEditor.ChannelButtonStyle")
+			.BorderBackgroundColor(this, &FTextureEditorToolkit::GetChannelButtonBackgroundColor, ETextureChannelButton::Blue)
+			.ForegroundColor(this, &FTextureEditorToolkit::GetChannelButtonForegroundColor, ETextureChannelButton::Blue)
+			.OnCheckStateChanged_Lambda(OnChannelCheckStateChanged, ETextureChannelButton::Blue)
+			.IsChecked(this, &FTextureEditorToolkit::OnGetChannelButtonCheckState, ETextureChannelButton::Blue)
+			[
+				SNew(STextBlock)
+				.Font(FAppStyle::Get().GetFontStyle("TextureEditor.ChannelButtonFont"))
+				.Text(FText::FromString("B"))
+			]
+		]
+		+SHorizontalBox::Slot()
+		.VAlign(VAlign_Center)
+		.Padding(2.0f)
+		.AutoWidth()
+		[
+			SNew(SCheckBox)
+			.Style(FAppStyle::Get(), "TextureEditor.ChannelButtonStyle")
+			.BorderBackgroundColor(this, &FTextureEditorToolkit::GetChannelButtonBackgroundColor, ETextureChannelButton::Alpha)
+			.ForegroundColor(this, &FTextureEditorToolkit::GetChannelButtonForegroundColor, ETextureChannelButton::Alpha)
+			.OnCheckStateChanged_Lambda(OnChannelCheckStateChanged, ETextureChannelButton::Alpha)
+			.IsChecked(this, &FTextureEditorToolkit::OnGetChannelButtonCheckState, ETextureChannelButton::Alpha)
+			.IsEnabled(this, &FTextureEditorToolkit::IsAlphaChannelButtonEnabled)
+			[
+				SNew(STextBlock)
+				.Font(FAppStyle::Get().GetFontStyle("TextureEditor.ChannelButtonFont"))
+				.Text(FText::FromString("A"))
+			]
+		];
+
+	return ChannelControl;
+}
+
+TSharedRef<SWidget> FTextureEditorToolkit::MakeLODControlWidget()
+{
+	TSharedRef<SWidget> LODControl = SNew(SBox)
+		.WidthOverride(212.0f)
+		[
+			SNew(SHorizontalBox)
+			.IsEnabled(this, &FTextureEditorToolkit::HandleMipLevelCheckBoxIsEnabled)
+			+ SHorizontalBox::Slot()
+			.Padding(4.0f, 0.0f, 2.0f, 0.0f)
+			.VAlign(VAlign_Center)
+			.AutoWidth()
+			[
+				SNew(SCheckBox)
+				.IsChecked(this, &FTextureEditorToolkit::HandleMipLevelCheckBoxIsChecked)
+				.OnCheckStateChanged(this, &FTextureEditorToolkit::HandleMipLevelCheckBoxCheckedStateChanged)
+			]
+			+ SHorizontalBox::Slot()
+			.VAlign(VAlign_Center)
+			.Padding(2.0f, 0.0f, 4.0f, 0.0f)
+			[
+				SNew(SComboButton)
+				.IsEnabled(this, &FTextureEditorToolkit::GetUseSpecifiedMip)
+				.OnGetMenuContent(this, &FTextureEditorToolkit::OnGenerateMipMapLevelMenu)
+				.ButtonContent()
+				[
+					SAssignNew(MipLevelTextBlock, STextBlock)
+					.Text(FText::Format(LOCTEXT("MipLevel", "Mip Level {0}"), SpecifiedMipLevel))
+				]
+			]
+			+ SHorizontalBox::Slot()
+			.AutoWidth()
+			.VAlign(VAlign_Center)
+			.Padding(4.0f)
+			[
+				SNew(SButton)
+				.ButtonStyle(FAppStyle::Get(), "TextureEditor.MipmapButtonStyle")
+				.OnClicked(this, &FTextureEditorToolkit::HandleMipMapPlusButtonClicked)
+				.IsEnabled(this, &FTextureEditorToolkit::GetUseSpecifiedMip)
+				[
+					SNew(SImage)
+					.Image(FAppStyle::Get().GetBrush("Icons.Plus"))
+					.ColorAndOpacity(FSlateColor::UseForeground())
+				]
+			]
+			+ SHorizontalBox::Slot()
+			.AutoWidth()
+			.VAlign(VAlign_Center)
+			.Padding(4.0f)
+			[
+				SNew(SButton)
+				.ButtonStyle(FAppStyle::Get(), "TextureEditor.MipmapButtonStyle")
+				.OnClicked(this, &FTextureEditorToolkit::HandleMipMapMinusButtonClicked)
+				.IsEnabled(this, &FTextureEditorToolkit::GetUseSpecifiedMip)
+				[
+					SNew(SImage)
+					.Image(FAppStyle::Get().GetBrush("Icons.Minus"))
+					.ColorAndOpacity(FSlateColor::UseForeground())
+				]
+			]
+		];
+
+	return LODControl;
+}
+
+TSharedRef<SWidget> FTextureEditorToolkit::MakeLayerControlWidget()
+{
+	TSharedRef<SWidget> LayerControl = SNew(SBox)
+		.WidthOverride(160.0f)
+		[
+			SNew(SHorizontalBox)
+			+ SHorizontalBox::Slot()
+			.Padding(4.0f, 0.0f, 4.0f, 0.0f)
+			.VAlign(VAlign_Center)
+			.AutoWidth()
+			[
+				SNew(STextBlock)
+				.Text(NSLOCTEXT("TextureEditor", "Layer", "Layer"))
+			]
+			+ SHorizontalBox::Slot()
+				.Padding(0.0f, 0.0f, 4.0f, 0.0f)
+				.VAlign(VAlign_Center)
+			[
+				SNew(SNumericEntryBox<int32>)
+				.AllowSpin(true)
+				.MinSliderValue(0)
+				.MaxSliderValue(this, &FTextureEditorToolkit::GetMaxLayer)
+				.Value(this, &FTextureEditorToolkit::HandleLayerEntryBoxValue)
+				.OnValueChanged(this, &FTextureEditorToolkit::HandleLayerEntryBoxChanged)
+			]
+		];
+
+	return LayerControl;
+}
+
+TSharedRef<SWidget> FTextureEditorToolkit::MakeExposureContolWidget()
+{
+	TSharedRef<SWidget> ExposureControl = SNew(SBox)
+		.WidthOverride(160.0f)
+		.Visibility(this, &FTextureEditorToolkit::HandleExposureBiasWidgetVisibility)
+		[
+			SNew(SHorizontalBox)
+			+ SHorizontalBox::Slot()
+			[
+				SNew(SHorizontalBox)
+				+ SHorizontalBox::Slot()
+				.Padding(8.0f, 0.0f, 4.0f, 0.0f)
+				.AutoWidth()
+				.VAlign(VAlign_Center)
+				[
+					SNew(STextBlock)
+					.Text(LOCTEXT("ExposureBiasLabel", "Exposure Bias"))
+				]
+				+ SHorizontalBox::Slot()
+				.Padding(0.0f, 0.0f, 4.0f, 0.0f)
+				.VAlign(VAlign_Center)
+				[
+					SNew(SNumericEntryBox<int32>)
+					.AllowSpin(true)
+					.MinSliderValue(MinExposure)
+					.MaxSliderValue(MaxExposure)
+					.Value(this, &FTextureEditorToolkit::HandleExposureBiasBoxValue)
+					.OnValueChanged(this, &FTextureEditorToolkit::HandleExposureBiasBoxValueChanged)
+				]
+			]
+		];
+	return ExposureControl;
+}
+
+TSharedRef<SWidget> FTextureEditorToolkit::MakeOpacityControlWidget()
+{
+	TSharedRef<SWidget> OpacityControl = SNew(SBox)
+		.WidthOverride(160.0f)
+		[
+			// opacity slider
+			SNew(SHorizontalBox)
+			+ SHorizontalBox::Slot()
+			.AutoWidth()
+			.VAlign(VAlign_Center)
+			[
+				SNew(STextBlock)
+				.Text(LOCTEXT("OpacityLabel", "Opacity"))
+			]
+			+ SHorizontalBox::Slot()
+			.Padding(4.0f, 0.0f)
+			.VAlign(VAlign_Center)
+			[
+				SNew(SNumericEntryBox<float>)
+				.AllowSpin(true)
+				.MinSliderValue(0.0f)
+				.MaxSliderValue(1.0f)
+				.OnValueChanged(this, &FTextureEditorToolkit::HandleOpacitySliderChanged)
+				.Value(this, &FTextureEditorToolkit::HandleOpacitySliderValue)
+			]
+		];
+
+	return OpacityControl;
+}
+
+TSharedRef<SWidget> FTextureEditorToolkit::MakeZoomControlWidget()
+{
+	const FMargin ToolbarSlotPadding(4.0f, 1.0f);
+	const FMargin ToolbarButtonPadding(4.0f, 0.0f);
+
+	FMenuBuilder ZoomMenuBuilder(true, NULL);
+	{
+		FUIAction Zoom25Action(FExecuteAction::CreateSP(this, &FTextureEditorToolkit::HandleZoomMenuEntryClicked, 0.25));
+		ZoomMenuBuilder.AddMenuEntry(LOCTEXT("Zoom25Action", "25%"), LOCTEXT("Zoom25ActionHint", "Show the texture at a quarter of its size."), FSlateIcon(), Zoom25Action);
+
+		FUIAction Zoom50Action(FExecuteAction::CreateSP(this, &FTextureEditorToolkit::HandleZoomMenuEntryClicked, 0.5));
+		ZoomMenuBuilder.AddMenuEntry(LOCTEXT("Zoom50Action", "50%"), LOCTEXT("Zoom50ActionHint", "Show the texture at half its size."), FSlateIcon(), Zoom50Action);
+
+		FUIAction Zoom100Action(FExecuteAction::CreateSP(this, &FTextureEditorToolkit::HandleZoomMenuEntryClicked, 1.0));
+		ZoomMenuBuilder.AddMenuEntry(LOCTEXT("Zoom100Action", "100%"), LOCTEXT("Zoom100ActionHint", "Show the texture in its original size."), FSlateIcon(), Zoom100Action);
+
+		FUIAction Zoom200Action(FExecuteAction::CreateSP(this, &FTextureEditorToolkit::HandleZoomMenuEntryClicked, 2.0));
+		ZoomMenuBuilder.AddMenuEntry(LOCTEXT("Zoom200Action", "200%"), LOCTEXT("Zoom200ActionHint", "Show the texture at twice its size."), FSlateIcon(), Zoom200Action);
+
+		FUIAction Zoom400Action(FExecuteAction::CreateSP(this, &FTextureEditorToolkit::HandleZoomMenuEntryClicked, 4.0));
+		ZoomMenuBuilder.AddMenuEntry(LOCTEXT("Zoom400Action", "400%"), LOCTEXT("Zoom400ActionHint", "Show the texture at four times its size."), FSlateIcon(), Zoom400Action);
+
+		ZoomMenuBuilder.AddMenuSeparator();
+
+		FUIAction ZoomFitAction(
+			FExecuteAction::CreateSP(this, &FTextureEditorToolkit::HandleZoomMenuFitClicked),
+			FCanExecuteAction(),
+			FIsActionChecked::CreateSP(this, &FTextureEditorToolkit::IsZoomMenuFitChecked)
+		);
+		ZoomMenuBuilder.AddMenuEntry(LOCTEXT("ZoomFitAction", "Scale To Fit"), LOCTEXT("ZoomFitActionHint", "Scales the texture down to fit within the viewport if needed."), FSlateIcon(), ZoomFitAction, NAME_None, EUserInterfaceActionType::RadioButton);
+
+		FUIAction ZoomFillAction(
+			FExecuteAction::CreateSP(this, &FTextureEditorToolkit::HandleZoomMenuFillClicked),
+			FCanExecuteAction(),
+			FIsActionChecked::CreateSP(this, &FTextureEditorToolkit::IsZoomMenuFillChecked)
+		);
+		ZoomMenuBuilder.AddMenuEntry(LOCTEXT("ZoomFillAction", "Scale To Fill"), LOCTEXT("ZoomFillActionHint", "Scales the texture up and down to fill the viewport."), FSlateIcon(), ZoomFillAction, NAME_None, EUserInterfaceActionType::RadioButton);
+	}
+
+	// zoom slider
+	TSharedRef<SWidget> ZoomControl = 
+		SNew(SHorizontalBox)
+		+ SHorizontalBox::Slot()
+		.AutoWidth()
+		.VAlign(VAlign_Center)
+		[
+			SNew(STextBlock)
+			.Text(LOCTEXT("ZoomLabel", "Zoom"))
+		]
+		+ SHorizontalBox::Slot()
+		.Padding(4.0f, 0.0f)
+		.AutoWidth()
+		.VAlign(VAlign_Center)
+		[
+			SNew(SBox)
+			.WidthOverride(200.f)
+			[
+				SNew(SSlider)
+				.OnValueChanged(this, &FTextureEditorToolkit::HandleZoomSliderChanged)
+				.Value(this, &FTextureEditorToolkit::HandleZoomSliderValue)
+			]
+		]
+		+ SHorizontalBox::Slot()
+		.AutoWidth()
+		.VAlign(VAlign_Center)
+		.Padding(4.0f, 0.0f, 0.0f, 0.0f)
+		[
+			SNew(SComboButton)
+			.ComboButtonStyle(FAppStyle::Get(), "SimpleComboButton")
+			.ButtonContent()
+			[
+				SNew(SBox)
+				.WidthOverride(40.f)
+				.HAlign(HAlign_Right)
+				[
+					SNew(STextBlock)
+					.Text(this, &FTextureEditorToolkit::HandleZoomPercentageText)
+				]
+			]
+			.MenuContent()
+			[
+				ZoomMenuBuilder.MakeWidget()
+			]
+		];
+
+	return ZoomControl;
+}
 
 #undef LOCTEXT_NAMESPACE
diff --git a/Engine/Source/Editor/TextureEditor/Private/TextureEditorToolkit.h b/Engine/Source/Editor/TextureEditor/Private/TextureEditorToolkit.h
--- a/Engine/Source/Editor/TextureEditor/Private/TextureEditorToolkit.h
+++ b/Engine/Source/Editor/TextureEditor/Private/TextureEditorToolkit.h
@@ -17,6 +17,14 @@ class STextureEditorViewport;
 class UFactory;
 class UTexture;
 
+enum class ETextureChannelButton : uint8
+{
+	Red,
+	Green,
+	Blue,
+	Alpha
+};
+
 /**
  * Implements an Editor toolkit for textures.
  */
@@ -45,17 +53,13 @@ public:
 	void InitTextureEditor( const EToolkitMode::Type Mode, const TSharedPtr<class IToolkitHost>& InitToolkitHost, UObject* ObjectToEdit );
 
 public:
-
 	// FAssetEditorToolkit interface
-
 	virtual FString GetDocumentationLink( ) const override;
 	virtual void RegisterTabSpawners( const TSharedRef<class FTabManager>& TabManager ) override;
 	virtual void UnregisterTabSpawners( const TSharedRef<class FTabManager>& TabManager ) override;
-
-public:
+	virtual bool UsesCustomToolbarPlacement() const { return false; }
 
 	// ITextureEditorToolkit interface
-
 	virtual void CalculateTextureDimensions( uint32& Width, uint32& Height ) const override;
 	virtual ESimpleElementBlendMode GetColourChannelBlendMode( ) const override;
 	virtual int32 GetMipLevel( ) const override;
@@ -76,31 +80,31 @@ public:
 	virtual void SetVolumeOpacity( float VolumeOpacity ) override;
 	virtual const FRotator& GetVolumeOrientation( ) const override;
 	virtual void SetVolumeOrientation( const FRotator& InOrientation ) override;
+	virtual bool IsVolumeTexture() const override
+	{
+		return bIsVolumeTexture;
+	}
 
-public:
+	virtual int32 GetExposureBias() const override
+	{
+		return ExposureBias;
+	}
 
-	// IToolkit interface
 
+	// IToolkit interface
 	virtual FText GetBaseToolkitName( ) const override;
 	virtual FName GetToolkitFName( ) const override;
 	virtual FLinearColor GetWorldCentricTabColorScale( ) const override;
 	virtual FString GetWorldCentricTabPrefix( ) const override;
 
-public:
-
 	// FGCObject interface
-
 	virtual void AddReferencedObjects( FReferenceCollector& Collector ) override;
 	
 protected:
-
 	// FEditorUndoClient interface
-
 	virtual void PostUndo( bool bSuccess ) override;
 	virtual void PostRedo( bool bSuccess ) override;
 
-protected:
-
 	/**
 	 * Binds the UI commands to delegates.
 	 */
@@ -123,7 +127,7 @@ protected:
 	 */
 	void ExtendToolBar( );
 
-	void FillToolbar(FToolBarBuilder& ToolbarBuilder, const TSharedRef< FUICommandList > ToolkitCommands, TSharedRef<SWidget> LODControl, TSharedRef<SWidget> LayerControl);
+	void FillToolbar(FToolBarBuilder& ToolbarBuilder);
 	/**
 	 * Gets the highest mip map level that this texture supports.
 	 *
@@ -138,22 +142,14 @@ protected:
 	 */
 	bool IsCubeTexture( ) const;
 
+	TSharedRef<SWidget> OnGenerateMipMapLevelMenu();
+	TSharedRef<SWidget> OnGenerateSettingsMenu();
 private:
-
-	// Callback for toggling the Alpha channel action.
-	void HandleAlphaChannelActionExecute( );
-
-	// Callback for getting the checked state of the Alpha channel action.
-	bool HandleAlphaChannelActionIsChecked( ) const;
-
-	// Callback for getting the enabled state of the Alpha channel action.
-	bool HandleAlphaChannelActionCanExecute( ) const;
-
-	// Callback for toggling the Blue channel action.
-	void HandleBlueChannelActionExecute( );
-
-	// Callback for getting the checked state of the Blue channel action.
-	bool HandleBlueChannelActionIsChecked( ) const;
+	bool IsAlphaChannelButtonEnabled() const;
+	FSlateColor GetChannelButtonBackgroundColor(ETextureChannelButton Button) const;
+	FSlateColor GetChannelButtonForegroundColor(ETextureChannelButton Button) const;
+	void OnChannelButtonCheckStateChanged(ETextureChannelButton Button);
+	ECheckBoxState OnGetChannelButtonCheckState(ETextureChannelButton Button) const;
 
 	// Callback for toggling the Checkered Background action.
 	void HandleCheckeredBackgroundActionExecute( ETextureEditorBackgrounds Background );
@@ -185,12 +181,6 @@ private:
 	// Callback for executing the Fit To 100% action.
 	void HandleZoomToNaturalActionExecute( );
 
-	// Callback for toggling the Green channel action.
-	void HandleGreenChannelActionExecute( );
-
-	// Callback for getting the checked state of the Green Channel action.
-	bool HandleGreenChannelActionIsChecked( ) const;
-
 	// Callback for changing the checked state of the MipMap check box.
 	void HandleMipLevelCheckBoxCheckedStateChanged( ECheckBoxState InNewState );
 
@@ -201,29 +191,19 @@ private:
 	bool HandleMipLevelCheckBoxIsEnabled( ) const;
 
 	// Callback for changing the value of the mip map level entry box.
-	void HandleMipLevelEntryBoxChanged( int32 NewMipLevel );
+	void HandleMipLevelChanged( int32 NewMipLevel );
 
 	// Callback for getting the value of the mip map level entry box.
 	TOptional<int32> HandleMipLevelEntryBoxValue( ) const;
 
-	// Callback for clicking the MipMinus button.
-	FReply HandleMipMapMinusButtonClicked( );
-
-	// Callback for clicking the MipPlus button.
-	FReply HandleMipMapPlusButtonClicked( );
+	FReply HandleMipMapMinusButtonClicked();
 
+	FReply HandleMipMapPlusButtonClicked();
 
 	void HandleLayerEntryBoxChanged(int32 NewMipLevel);
 	TOptional<int32> HandleLayerEntryBoxValue() const;
-	FReply HandleLayerMinusButtonClicked();
-	FReply HandleLayerPlusButtonClicked();
-	bool HasLayers() const;
-
-	// Callback for toggling the Red channel action.
-	void HandleRedChannelActionExecute( );
 
-	// Callback for getting the checked state of the Red Channel action.
-	bool HandleRedChannelActionIsChecked( ) const;
+	bool HasLayers() const;
 
 	// Callback for determining whether the Reimport action can execute.
 	bool HandleReimportActionCanExecute( ) const;
@@ -261,6 +241,54 @@ private:
 	// Callback for getting the checked state of the Texture Border action.
 	bool HandleTextureBorderActionIsChecked( ) const;
 
+	// Callback for getting the visibility of the exposure bias widget.
+	EVisibility HandleExposureBiasWidgetVisibility() const;
+
+	// Callback for getting the exposure bias.
+	TOptional<int32> HandleExposureBiasBoxValue() const;
+
+	// Callback for changing the exposure bias.
+	void HandleExposureBiasBoxValueChanged(int32 NewExposure);
+
+	// Callback for changes in the zoom slider.
+	void HandleOpacitySliderChanged(float NewValue);
+
+	// Callback for getting the zoom slider's value.
+	TOptional<float> HandleOpacitySliderValue() const;
+
+	// Callback for clicking an item in the 'Zoom' menu.
+	void HandleZoomMenuEntryClicked(double ZoomValue);
+
+	// Callback for clicking the 'Fill' item in the 'Zoom' menu.
+	void HandleZoomMenuFillClicked();
+
+	// Callback for setting the checked state of the 'Fill' item in the 'Zoom' menu.
+	bool IsZoomMenuFillChecked() const;
+
+	// Callback for clicking the 'Fit' item in the 'Zoom' menu.
+	void HandleZoomMenuFitClicked();
+
+	// Callback for setting the checked state of the 'Fit' item in the 'Zoom' menu.
+	bool IsZoomMenuFitChecked() const;
+
+	// Callback for getting the zoom percentage text.
+	FText HandleZoomPercentageText() const;
+
+	// Callback for changes in the zoom slider.
+	void HandleZoomSliderChanged(float NewValue);
+
+	// Callback for getting the zoom slider's value.
+	float HandleZoomSliderValue() const;
+
+	// Callback for clicking the View Options menu button.
+	FReply HandleViewOptionsMenuButtonClicked();
+
+	TSharedRef<SWidget> MakeChannelControlWidget();
+	TSharedRef<SWidget> MakeLODControlWidget();
+	TSharedRef<SWidget> MakeLayerControlWidget();
+	TSharedRef<SWidget> MakeExposureContolWidget();
+	TSharedRef<SWidget> MakeOpacityControlWidget();
+	TSharedRef<SWidget> MakeZoomControlWidget();
 private:
 
 	/** The Texture asset being inspected */
@@ -285,6 +313,10 @@ private:
 	TSharedPtr<STextBlock> LODBiasText;
 	TSharedPtr<STextBlock> HasAlphaChannelText;
 	TSharedPtr<STextBlock> NumMipsText;
+	TSharedPtr<STextBlock> MipLevelTextBlock;
+
+	// Holds the anchor for the view options menu.
+	TSharedPtr<SMenuAnchor> ViewOptionsMenuAnchor;
 
 	/** If true, displays the red channel */
 	bool bIsRedChannel;
@@ -318,6 +350,9 @@ private:
 	/** The texture's zoom factor. */
 	double Zoom;
 
+	// Which exposure level should be used, in FStop e.g. 0:original, -1:half as bright, 1:2x as bright, 2:4x as bright.
+	int32 ExposureBias;
+
 	/** This toolkit's current zoom mode **/
 	ETextureEditorZoomMode ZoomMode;
 
@@ -327,6 +362,7 @@ private:
 	// For volume texture, the orientation when tracing.
 	FRotator VolumeOrientation;
 
+	bool bIsVolumeTexture;
 private:
 
 	// The name of the Viewport tab.
diff --git a/Engine/Source/Editor/TextureEditor/Private/Widgets/STextureEditorViewport.cpp b/Engine/Source/Editor/TextureEditor/Private/Widgets/STextureEditorViewport.cpp
--- a/Engine/Source/Editor/TextureEditor/Private/Widgets/STextureEditorViewport.cpp
+++ b/Engine/Source/Editor/TextureEditor/Private/Widgets/STextureEditorViewport.cpp
@@ -14,18 +14,12 @@
 #include "Engine/TextureRenderTargetVolume.h"
 #include "Slate/SceneViewport.h"
 #include "TextureEditorConstants.h"
-#include "Widgets/STextureEditorViewportToolbar.h"
 #include "Widgets/Input/SNumericEntryBox.h"
 #include "TextureEditorSettings.h"
 
 
 #define LOCTEXT_NAMESPACE "STextureEditorViewport"
 
-// Specifies the maximum allowed exposure bias.
-const int32 MaxExposure = 10;
-
-// Specifies the minimum allowed exposure bias.
-const int32 MinExposure = -10;
 
 
 /* STextureEditorViewport interface
@@ -39,54 +33,16 @@ void STextureEditorViewport::AddReferencedObjects( FReferenceCollector& Collecto
 
 void STextureEditorViewport::Construct( const FArguments& InArgs, const TSharedRef<ITextureEditorToolkit>& InToolkit )
 {
-	ExposureBias = 0;
+
 	bIsRenderingEnabled = true;
 	ToolkitPtr = InToolkit;
 	
-	// create zoom menu
-	FMenuBuilder ZoomMenuBuilder(true, NULL);
-	{
-		FUIAction Zoom25Action(FExecuteAction::CreateSP(this, &STextureEditorViewport::HandleZoomMenuEntryClicked, 0.25));
-		ZoomMenuBuilder.AddMenuEntry(LOCTEXT("Zoom25Action", "25%"), LOCTEXT("Zoom25ActionHint", "Show the texture at a quarter of its size."), FSlateIcon(), Zoom25Action);
-
-		FUIAction Zoom50Action(FExecuteAction::CreateSP(this, &STextureEditorViewport::HandleZoomMenuEntryClicked, 0.5));
-		ZoomMenuBuilder.AddMenuEntry(LOCTEXT("Zoom50Action", "50%"), LOCTEXT("Zoom50ActionHint", "Show the texture at half its size."), FSlateIcon(), Zoom50Action);
-
-		FUIAction Zoom100Action(FExecuteAction::CreateSP(this, &STextureEditorViewport::HandleZoomMenuEntryClicked, 1.0));
-		ZoomMenuBuilder.AddMenuEntry(LOCTEXT("Zoom100Action", "100%"), LOCTEXT("Zoom100ActionHint", "Show the texture in its original size."), FSlateIcon(), Zoom100Action);
-
-		FUIAction Zoom200Action(FExecuteAction::CreateSP(this, &STextureEditorViewport::HandleZoomMenuEntryClicked, 2.0));
-		ZoomMenuBuilder.AddMenuEntry(LOCTEXT("Zoom200Action", "200%"), LOCTEXT("Zoom200ActionHint", "Show the texture at twice its size."), FSlateIcon(), Zoom200Action);
-
-		FUIAction Zoom400Action(FExecuteAction::CreateSP(this, &STextureEditorViewport::HandleZoomMenuEntryClicked, 4.0));
-		ZoomMenuBuilder.AddMenuEntry(LOCTEXT("Zoom400Action", "400%"), LOCTEXT("Zoom400ActionHint", "Show the texture at four times its size."), FSlateIcon(), Zoom400Action);
-
-		ZoomMenuBuilder.AddMenuSeparator();
-
-		FUIAction ZoomFitAction(
-			FExecuteAction::CreateSP(this, &STextureEditorViewport::HandleZoomMenuFitClicked),
-			FCanExecuteAction(),
-			FIsActionChecked::CreateSP(this, &STextureEditorViewport::IsZoomMenuFitChecked)
-		);
-		ZoomMenuBuilder.AddMenuEntry(LOCTEXT("ZoomFitAction", "Scale To Fit"), LOCTEXT("ZoomFitActionHint", "Scales the texture down to fit within the viewport if needed."), FSlateIcon(), ZoomFitAction, NAME_None, EUserInterfaceActionType::RadioButton);
-
-		FUIAction ZoomFillAction(
-			FExecuteAction::CreateSP(this, &STextureEditorViewport::HandleZoomMenuFillClicked),
-			FCanExecuteAction(),
-			FIsActionChecked::CreateSP(this, &STextureEditorViewport::IsZoomMenuFillChecked)
-		);
-		ZoomMenuBuilder.AddMenuEntry(LOCTEXT("ZoomFillAction", "Scale To Fill"), LOCTEXT("ZoomFillActionHint", "Scales the texture up and down to fill the viewport."), FSlateIcon(), ZoomFillAction, NAME_None, EUserInterfaceActionType::RadioButton);
-	}
-
 	FText TextureName = FText::GetEmpty();
-	
-	bool bIsVolumeTexture = false;
+
 	if (InToolkit->GetTexture() != nullptr)
 	{
 		FText FormattedText = InToolkit->HasValidTextureResource() ? FText::FromString(TEXT("{0}")) : LOCTEXT( "InvalidTexture", "{0} (Invalid Texture)");
 		TextureName = FText::Format(FormattedText, FText::FromName(InToolkit->GetTexture()->GetFName()));
-
-		bIsVolumeTexture = InToolkit->GetTexture()->IsA<UVolumeTexture>() || InToolkit->GetTexture()->IsA<UTextureRenderTargetVolume>();
 	}
 
 
@@ -95,198 +51,44 @@ void STextureEditorViewport::Construct( const FArguments& InArgs, const TSharedR
 	this->ChildSlot
 	[
 		SNew(SVerticalBox)
-
 		+ SVerticalBox::Slot()
-			.FillHeight(1.0f)
+		.FillHeight(1.0f)
+		[
+			SNew(SOverlay)
+			// viewport canvas
+			+ SOverlay::Slot()
+			.Padding(5.0f, 0.0f)
 			[
-				SNew(SHorizontalBox)
-
-				+ SHorizontalBox::Slot()
-					.FillWidth(1.0f)
-					[
-						SNew(SVerticalBox)
-
-						+ SVerticalBox::Slot()
-							.FillHeight(1)
-							[
-								SNew(SOverlay)
-
-								// viewport canvas
-								+ SOverlay::Slot()
-									[
-										SAssignNew(ViewportWidget, SViewport)
-											.EnableGammaCorrection(false)
-											.IsEnabled(FSlateApplication::Get().GetNormalExecutionAttribute())
-											.ShowEffectWhenDisabled(false)
-											.EnableBlending(true)
-											.ToolTip(SNew(SToolTip).Text(this, &STextureEditorViewport::GetDisplayedResolution))
-									]
-	
-								// tool bar
-								+ SOverlay::Slot()
-									.Padding(2.0f)
-									.VAlign(VAlign_Top)
-									[
-										SNew(SHorizontalBox)
-
-										+ SHorizontalBox::Slot()
-											.AutoWidth()
-											[
-												SNew(STextureEditorViewportToolbar, InToolkit->GetToolkitCommands())
-													.IsVolumeTexture(bIsVolumeTexture)
-											]
-
-										+ SHorizontalBox::Slot()
-											.FillWidth(1.0f)
-											.Padding(4.0f, 0.0f)
-											.HAlign(HAlign_Right)
-											.VAlign(VAlign_Center)
-											[
-												SNew(STextBlock)
-													.Text(TextureName)
-											]
-									]
-							]
-					]
-
-				+ SHorizontalBox::Slot()
-					.AutoWidth()
-					[
-						// vertical scroll bar
-						SAssignNew(TextureViewportVerticalScrollBar, SScrollBar)
-							.Visibility(this, &STextureEditorViewport::HandleVerticalScrollBarVisibility)
-							.OnUserScrolled(this, &STextureEditorViewport::HandleVerticalScrollBarScrolled)
-					]
+				SAssignNew(ViewportWidget, SViewport)
+				.EnableGammaCorrection(false)
+				.IsEnabled(FSlateApplication::Get().GetNormalExecutionAttribute())
+				.ShowEffectWhenDisabled(false)
+				.EnableBlending(true)
+				.ToolTip(SNew(SToolTip).Text(this, &STextureEditorViewport::GetDisplayedResolution))
 			]
-
-		+ SVerticalBox::Slot()
-			.AutoHeight()
+			+ SOverlay::Slot()
+			.HAlign(HAlign_Right)
 			[
-				// horizontal scrollbar
-				SAssignNew(TextureViewportHorizontalScrollBar, SScrollBar)
-					.Orientation( Orient_Horizontal )
-					.Visibility(this, &STextureEditorViewport::HandleHorizontalScrollBarVisibility)
-					.OnUserScrolled(this, &STextureEditorViewport::HandleHorizontalScrollBarScrolled)
+				// vertical scroll bar
+				SAssignNew(TextureViewportVerticalScrollBar, SScrollBar)
+				.Visibility(this, &STextureEditorViewport::HandleVerticalScrollBarVisibility)
+				.OnUserScrolled(this, &STextureEditorViewport::HandleVerticalScrollBarScrolled)
 			]
-
+		]
 		+ SVerticalBox::Slot()
-			.AutoHeight()
-			.Padding(0.0f, 2.0f, 0.0f, 0.0f)
-			[
-				SAssignNew(HorizontalBox, SHorizontalBox)
-
-				// exposure bias
-				+ SHorizontalBox::Slot()
-					.FillWidth(0.3f)
-					[
-						SNew(SHorizontalBox)
-							.Visibility(this, &STextureEditorViewport::HandleExposureBiasWidgetVisibility)
-
-						+ SHorizontalBox::Slot()
-							.AutoWidth()
-							.VAlign(VAlign_Center)
-							[
-								SNew(STextBlock)
-									.Text(LOCTEXT("ExposureBiasLabel", "Exposure Bias:"))
-							]
-
-						+ SHorizontalBox::Slot()
-							.AutoWidth()
-							.Padding(4.0f, 0.0f)
-							[
-								SNew(SNumericEntryBox<int32>)
-									.AllowSpin(true)
-									.MinSliderValue(MinExposure)
-									.MaxSliderValue(MaxExposure)
-									.Value(this, &STextureEditorViewport::HandleExposureBiasBoxValue)
-									.OnValueChanged(this, &STextureEditorViewport::HandleExposureBiasBoxValueChanged)
-							]
-					]
-
-				// separator
-				+ SHorizontalBox::Slot()
-					.FillWidth(0.3f)
-					[
-						SNullWidget::NullWidget
-					]
-			]
-	];
-
+		.AutoHeight()
+		[
+			// horizontal scrollbar
+			SAssignNew(TextureViewportHorizontalScrollBar, SScrollBar)
+				.Orientation( Orient_Horizontal )
+				.Visibility(this, &STextureEditorViewport::HandleHorizontalScrollBarVisibility)
+				.OnUserScrolled(this, &STextureEditorViewport::HandleHorizontalScrollBarScrolled)
+		]
 
-	if (bIsVolumeTexture)
-	{
-		// opacity slider
-		HorizontalBox->AddSlot()
-			.FillWidth(0.3f)
-			[
-				SNew(SHorizontalBox)
-
-				+ SHorizontalBox::Slot()
-					.AutoWidth()
-					.VAlign(VAlign_Center)
-					[
-						SNew(STextBlock)
-							.Text(LOCTEXT("OpacityLabel", "Opacity:"))
-					]
-
-				+ SHorizontalBox::Slot()
-					.FillWidth(1.0f)
-					.Padding(4.0f, 0.0f)
-					.VAlign(VAlign_Center)
-					[
-						SNew(SSlider)
-							.OnValueChanged(this, &STextureEditorViewport::HandleOpacitySliderChanged)
-							.Value(this, &STextureEditorViewport::HandleOpacitySliderValue)
-					]
-			];
-	}
+	];
 
-	// zoom slider
-	HorizontalBox->AddSlot()
-		.FillWidth(0.3f)
-		[
-			SNew(SHorizontalBox)
-
-			+ SHorizontalBox::Slot()
-				.AutoWidth()
-				.VAlign(VAlign_Center)
-				[
-					SNew(STextBlock)
-						.Text(LOCTEXT("ZoomLabel", "Zoom:"))
-				]
-
-			+ SHorizontalBox::Slot()
-				.FillWidth(1.0f)
-				.Padding(4.0f, 0.0f)
-				.VAlign(VAlign_Center)
-				[
-					SNew(SSlider)
-						.OnValueChanged(this, &STextureEditorViewport::HandleZoomSliderChanged)
-						.Value(this, &STextureEditorViewport::HandleZoomSliderValue)
-				]
-
-			+ SHorizontalBox::Slot()
-				.AutoWidth()
-				.VAlign(VAlign_Center)
-				[
-					SNew(STextBlock)
-						.Text(this, &STextureEditorViewport::HandleZoomPercentageText)
-				]
-
-			+ SHorizontalBox::Slot()
-				.AutoWidth()
-				.Padding(4.0f, 0.0f, 0.0f, 0.0f)
-				.VAlign(VAlign_Center)
-				[
-					SNew(SComboButton)
-						.ContentPadding(FMargin(0.0))
-						.MenuContent()
-						[
-							ZoomMenuBuilder.MakeWidget()
-						]
-				]
-		];
 
+	
 	ViewportClient = MakeShareable(new FTextureEditorViewportClient(ToolkitPtr, SharedThis(this)));
 
 	Viewport = MakeShareable(new FSceneViewport(ViewportClient.Get(), ViewportWidget));
@@ -359,30 +161,6 @@ FText STextureEditorViewport::GetDisplayedResolution( ) const
 /* STextureEditorViewport event handlers
  *****************************************************************************/
 
-EVisibility STextureEditorViewport::HandleExposureBiasWidgetVisibility( ) const
-{
-	UTexture* Texture = ToolkitPtr.Pin()->GetTexture();
-
-	if ((Texture != NULL) && (Texture->CompressionSettings == TC_HDR || Texture->CompressionSettings == TC_HDR_Compressed))
-	{
-		return EVisibility::Visible;
-	}
-
-	return EVisibility::Hidden;
-}
-
-
-TOptional<int32> STextureEditorViewport::HandleExposureBiasBoxValue( ) const
-{
-	return ExposureBias;
-}
-
-
-void STextureEditorViewport::HandleExposureBiasBoxValueChanged( int32 NewExposure )
-{
-	ExposureBias = NewExposure;
-}
-
 
 void STextureEditorViewport::HandleHorizontalScrollBarScrolled( float InScrollOffsetFraction )
 {
@@ -426,82 +204,6 @@ EVisibility STextureEditorViewport::HandleVerticalScrollBarVisibility( ) const
 }
 
 
-void STextureEditorViewport::HandleZoomMenuEntryClicked( double ZoomValue )
-{
-	ToolkitPtr.Pin()->SetCustomZoomLevel(ZoomValue);
-}
-
-
-void STextureEditorViewport::HandleZoomMenuFillClicked()
-{
-	ToolkitPtr.Pin()->SetZoomMode(ETextureEditorZoomMode::Fill);
-}
-
-void STextureEditorViewport::HandleZoomMenuFitClicked()
-{
-	ToolkitPtr.Pin()->SetZoomMode(ETextureEditorZoomMode::Fit);
-}
-
-
-bool STextureEditorViewport::IsZoomMenuFillChecked() const
-{
-	return ToolkitPtr.Pin()->IsCurrentZoomMode(ETextureEditorZoomMode::Fill);
-}
-
-bool STextureEditorViewport::IsZoomMenuFitChecked() const
-{
-	return ToolkitPtr.Pin()->IsCurrentZoomMode(ETextureEditorZoomMode::Fit);
-}
-
-bool STextureEditorViewport::HasValidTextureResource() const
-{
-	return ToolkitPtr.Pin()->HasValidTextureResource();
-}
-
-FText STextureEditorViewport::HandleZoomPercentageText( ) const
-{
-	double DisplayedZoomLevel = ToolkitPtr.Pin()->CalculateDisplayedZoomLevel();
-	FText ZoomLevelPercent = FText::AsPercent(DisplayedZoomLevel);
-
-	// For fit and fill, show the effective zoom level in parenthesis - eg. "Fill (220%)"
-	static const FText ZoomModeWithPercentFormat = LOCTEXT("ZoomModeWithPercentFormat", "{ZoomMode} ({ZoomPercent})");
-	if (ToolkitPtr.Pin()->GetZoomMode() == ETextureEditorZoomMode::Fit)
-	{
-		static const FText ZoomModeFit = LOCTEXT("ZoomModeFit", "Fit");
-		return FText::FormatNamed(ZoomModeWithPercentFormat, TEXT("ZoomMode"), ZoomModeFit, TEXT("ZoomPercent"), ZoomLevelPercent);
-	}
-
-	if (ToolkitPtr.Pin()->GetZoomMode() == ETextureEditorZoomMode::Fill)
-	{
-		static const FText ZoomModeFill = LOCTEXT("ZoomModeFill", "Fill");
-		return FText::FormatNamed(ZoomModeWithPercentFormat, TEXT("ZoomMode"), ZoomModeFill, TEXT("ZoomPercent"), ZoomLevelPercent);
-	}
-
-	// If custom, then just the percent is enough
-	return ZoomLevelPercent;
-}
-
-
-void STextureEditorViewport::HandleZoomSliderChanged( float NewValue )
-{
-	ToolkitPtr.Pin()->SetCustomZoomLevel(NewValue * MaxZoom);
-}
-
-
-float STextureEditorViewport::HandleZoomSliderValue( ) const
-{
-	return (ToolkitPtr.Pin()->CalculateDisplayedZoomLevel() / MaxZoom);
-}
-
-void STextureEditorViewport::HandleOpacitySliderChanged(float NewValue)
-{
-	ToolkitPtr.Pin()->SetVolumeOpacity(NewValue);
-}
-
-float STextureEditorViewport::HandleOpacitySliderValue() const
-{
-	return ToolkitPtr.Pin()->GetVolumeOpacity();
-}
 
 
 #undef LOCTEXT_NAMESPACE
diff --git a/Engine/Source/Editor/TextureEditor/Private/Widgets/STextureEditorViewport.h b/Engine/Source/Editor/TextureEditor/Private/Widgets/STextureEditorViewport.h
--- a/Engine/Source/Editor/TextureEditor/Private/Widgets/STextureEditorViewport.h
+++ b/Engine/Source/Editor/TextureEditor/Private/Widgets/STextureEditorViewport.h
@@ -42,15 +42,6 @@ public:
 	 */
 	void ModifyCheckerboardTextureColors( );
 
-	/**
-	 * Gets the exposure bias.
-	 *
-	 * @return Exposure bias value.
-	 */
-	int32 GetExposureBias( ) const
-	{
-		return ExposureBias;
-	}
 
 	/** Enable viewport rendering */
 	void EnableRendering();
@@ -80,15 +71,6 @@ protected:
 
 private:
 
-	// Callback for getting the visibility of the exposure bias widget.
-	EVisibility HandleExposureBiasWidgetVisibility( ) const;
-
-	// Callback for getting the exposure bias.
-	TOptional<int32> HandleExposureBiasBoxValue( ) const;
-
-	// Callback for changing the exposure bias.
-	void HandleExposureBiasBoxValueChanged( int32 NewExposure );
-
 	// Callback for the horizontal scroll bar.
 	void HandleHorizontalScrollBarScrolled( float InScrollOffsetFraction );
 
@@ -101,43 +83,12 @@ private:
 	// Callback for getting the visibility of the horizontal scroll bar.
 	EVisibility HandleVerticalScrollBarVisibility( ) const;
 
-	// Callback for clicking an item in the 'Zoom' menu.
-	void HandleZoomMenuEntryClicked( double ZoomValue );
-
-	// Callback for clicking the 'Fill' item in the 'Zoom' menu.
-	void HandleZoomMenuFillClicked();
-
-	// Callback for setting the checked state of the 'Fill' item in the 'Zoom' menu.
-	bool IsZoomMenuFillChecked() const;
-
-	// Callback for clicking the 'Fit' item in the 'Zoom' menu.
-	void HandleZoomMenuFitClicked();
-
-	// Callback for setting the checked state of the 'Fit' item in the 'Zoom' menu.
-	bool IsZoomMenuFitChecked() const;
-
-	// Callback for getting the zoom percentage text.
-	FText HandleZoomPercentageText( ) const;
-
-	// Callback for changes in the zoom slider.
-	void HandleZoomSliderChanged( float NewValue );
-
-	// Callback for getting the zoom slider's value.
-	float HandleZoomSliderValue( ) const;
-
-	// Callback for changes in the zoom slider.
-	void HandleOpacitySliderChanged( float NewValue );
-
-	// Callback for getting the zoom slider's value.
-	float HandleOpacitySliderValue( ) const;
 
 	// Checks if the texture being edited has a valid texture resource
 	bool HasValidTextureResource( ) const;
 
 private:
 
-	// Which exposure level should be used, in FStop e.g. 0:original, -1:half as bright, 1:2x as bright, 2:4x as bright.
-	int32 ExposureBias;
 
 	// Pointer back to the Texture editor tool that owns us.
 	TWeakPtr<ITextureEditorToolkit> ToolkitPtr;
diff --git a/Engine/Source/Editor/TextureEditor/Private/Widgets/STextureEditorViewportToolbar.cpp b/Engine/Source/Editor/TextureEditor/Private/Widgets/STextureEditorViewportToolbar.cpp
deleted file mode 100644
diff --git a/Engine/Source/Editor/TextureEditor/Private/Widgets/STextureEditorViewportToolbar.h b/Engine/Source/Editor/TextureEditor/Private/Widgets/STextureEditorViewportToolbar.h
deleted file mode 100644
diff --git a/Engine/Source/Editor/TextureEditor/Public/Interfaces/ITextureEditorToolkit.h b/Engine/Source/Editor/TextureEditor/Public/Interfaces/ITextureEditorToolkit.h
--- a/Engine/Source/Editor/TextureEditor/Public/Interfaces/ITextureEditorToolkit.h
+++ b/Engine/Source/Editor/TextureEditor/Public/Interfaces/ITextureEditorToolkit.h
@@ -45,31 +45,8 @@ public:
 	virtual void SetVolumeOpacity( float VolumeOpacity ) = 0;
 	virtual const FRotator& GetVolumeOrientation( ) const = 0;
 	virtual void SetVolumeOrientation( const FRotator& InOrientation ) = 0;
-
-	UE_DEPRECATED(4.26, "There are now more than 2 zoom modes, use GetZoomMode() instead.")
-	virtual bool GetFitToViewport() const final
-	{
-		return GetZoomMode() != ETextureEditorZoomMode::Custom;
-	}
-	
-	UE_DEPRECATED(4.26, "There are now more than 2 zoom modes, use SetZoomMode() instead.")
-	virtual void SetFitToViewport(const bool bFitToViewport) final
-	{
-		SetZoomMode(bFitToViewport ? ETextureEditorZoomMode::Fit : ETextureEditorZoomMode::Custom);
-	}
-
-	UE_DEPRECATED(4.26, "GetZoom Has been renamed to GetCustomZoomLevel to clarify it refers to the custom zoom mode's value.")
-	virtual double GetZoom() const final
-	{
-		return GetCustomZoomLevel();
-	}
-
-	UE_DEPRECATED(4.26, "SetZoom has been renamed to SetCustomZoomLevel to clarify it refers to the custom zoom mode's value.")
-	virtual void SetZoom(double ZoomValue) final
-	{
-		SetCustomZoomLevel(ZoomValue);
-	}
-
+	virtual int32 GetExposureBias() const = 0;
+	virtual bool IsVolumeTexture() const = 0;
 public:
 
 	/**
diff --git a/Engine/Source/Runtime/SlateCore/Private/Styling/StarshipCoreStyle.cpp b/Engine/Source/Runtime/SlateCore/Private/Styling/StarshipCoreStyle.cpp
--- a/Engine/Source/Runtime/SlateCore/Private/Styling/StarshipCoreStyle.cpp
+++ b/Engine/Source/Runtime/SlateCore/Private/Styling/StarshipCoreStyle.cpp
@@ -24,41 +24,6 @@
 
 TSharedPtr< ISlateStyle > FStarshipCoreStyle::Instance = nullptr;
 
-namespace CoreStyleConstants
-{
-	// Note, these sizes are in Slate Units.
-	// Slate Units do NOT have to map to pixels.
-	const FVector2D Icon5x16(5.0f, 16.0f);
-	const FVector2D Icon8x4(8.0f, 4.0f);
-	const FVector2D Icon16x4(16.0f, 4.0f);
-	const FVector2D Icon8x8(8.0f, 8.0f);
-	const FVector2D Icon10x10(10.0f, 10.0f);
-	const FVector2D Icon12x12(12.0f, 12.0f);
-	const FVector2D Icon12x16(12.0f, 16.0f);
-	const FVector2D Icon14x14(14.0f, 14.0f);
-	const FVector2D Icon16x16(16.0f, 16.0f);
-	const FVector2D Icon18x18(18.0f, 18.0f);
-	const FVector2D Icon20x20(20.0f, 20.0f);
-	const FVector2D Icon22x22(22.0f, 22.0f);
-	const FVector2D Icon24x24(24.0f, 24.0f);
-	const FVector2D Icon25x25(25.0f, 25.0f);
-	const FVector2D Icon32x32(32.0f, 32.0f);
-	const FVector2D Icon40x40(40.0f, 40.0f);
-	const FVector2D Icon64x64(64.0f, 64.0f);
-	const FVector2D Icon36x24(36.0f, 24.0f);
-	const FVector2D Icon128x128(128.0f, 128.0f);
-
-	// Common Margins
-	const FMargin DefaultMargins(8.f, 4.f);
-	// Buttons already have a built in (4., 2.) padding - adding to that a little
-	const FMargin ButtonMargins(20.f, 2.f, 20.f, 3.f);
-
-	const float InputFocusRadius = 4.f;
-	const float InputFocusThickness = 1.0f;
-
-	const FLinearColor WindowHighlight(FColor::FromHex("#4C4C4CFF"));
-}
-
 using namespace CoreStyleConstants;
 
 #define FONT(...) FSlateFontInfo(FLegacySlateFontInfoCache::Get().GetDefaultFont(), __VA_ARGS__)
@@ -336,6 +301,7 @@ TSharedRef<ISlateStyle> FStarshipCoreStyle::Create()
 		Style->Set("Icons.Blueprints", new IMAGE_BRUSH_SVG("Starship/Common/blueprint", Icon16x16));
 		Style->Set("Icons.Cross", new IMAGE_BRUSH_SVG("Starship/Common/plus", Icon16x16));
 		Style->Set("Icons.Plus", new IMAGE_BRUSH_SVG("Starship/Common/plus", Icon16x16));
+		Style->Set("Icons.Minus", new IMAGE_BRUSH_SVG("Starship/Common/minus", Icon16x16));
 		Style->Set("Icons.PlusCircle", new IMAGE_BRUSH_SVG("Starship/Common/plus-circle", Icon16x16));
 		Style->Set("Icons.X", new IMAGE_BRUSH_SVG("Starship/Common/close", Icon16x16));
 
diff --git a/Engine/Source/Runtime/SlateCore/Public/Styling/StarshipCoreStyle.h b/Engine/Source/Runtime/SlateCore/Public/Styling/StarshipCoreStyle.h
--- a/Engine/Source/Runtime/SlateCore/Public/Styling/StarshipCoreStyle.h
+++ b/Engine/Source/Runtime/SlateCore/Public/Styling/StarshipCoreStyle.h
@@ -10,6 +10,41 @@ struct FSlateDynamicImageBrush;
 
 class FStyle;
 
+namespace CoreStyleConstants
+{
+	// Note, these sizes are in Slate Units.
+	// Slate Units do NOT have to map to pixels.
+	const FVector2D Icon5x16(5.0f, 16.0f);
+	const FVector2D Icon8x4(8.0f, 4.0f);
+	const FVector2D Icon16x4(16.0f, 4.0f);
+	const FVector2D Icon8x8(8.0f, 8.0f);
+	const FVector2D Icon10x10(10.0f, 10.0f);
+	const FVector2D Icon12x12(12.0f, 12.0f);
+	const FVector2D Icon12x16(12.0f, 16.0f);
+	const FVector2D Icon14x14(14.0f, 14.0f);
+	const FVector2D Icon16x16(16.0f, 16.0f);
+	const FVector2D Icon18x18(18.0f, 18.0f);
+	const FVector2D Icon20x20(20.0f, 20.0f);
+	const FVector2D Icon22x22(22.0f, 22.0f);
+	const FVector2D Icon24x24(24.0f, 24.0f);
+	const FVector2D Icon25x25(25.0f, 25.0f);
+	const FVector2D Icon32x32(32.0f, 32.0f);
+	const FVector2D Icon40x40(40.0f, 40.0f);
+	const FVector2D Icon64x64(64.0f, 64.0f);
+	const FVector2D Icon36x24(36.0f, 24.0f);
+	const FVector2D Icon128x128(128.0f, 128.0f);
+
+	// Common Margins
+	const FMargin DefaultMargins(8.f, 4.f);
+	// Buttons already have a built in (4., 2.) padding - adding to that a little
+	const FMargin ButtonMargins(20.f, 2.f, 20.f, 3.f);
+
+	const float InputFocusRadius = 4.f;
+	const float InputFocusThickness = 1.0f;
+
+	const FLinearColor WindowHighlight(FColor::FromHex("#4C4C4CFF"));
+}
+
 /**
  * Core slate style
  */
