Removing redundant conditional

diff --git a/Engine/Source/Editor/TextureEditor/Private/TextureEditorToolkit.cpp b/Engine/Source/Editor/TextureEditor/Private/TextureEditorToolkit.cpp
--- a/Engine/Source/Editor/TextureEditor/Private/TextureEditorToolkit.cpp
+++ b/Engine/Source/Editor/TextureEditor/Private/TextureEditorToolkit.cpp
@@ -1206,7 +1206,7 @@ FSlateColor FTextureEditorToolkit::GetChannelButtonBackgroundColor(ETextureChann
 	case ETextureChannelButton::Blue:
 		return bIsBlueChannel ? FLinearColor::Blue : FLinearColor::White;
 	case ETextureChannelButton::Alpha:
-		return bIsAlphaChannel ? FLinearColor::White : FLinearColor::White;
+		return FLinearColor::White;
 	default:
 		check(false);
 		return FSlateColor();
