Fixed cut off text in texture editor toolbar

diff --git a/Engine/Source/Editor/TextureEditor/Private/TextureEditorToolkit.cpp b/Engine/Source/Editor/TextureEditor/Private/TextureEditorToolkit.cpp
--- a/Engine/Source/Editor/TextureEditor/Private/TextureEditorToolkit.cpp
+++ b/Engine/Source/Editor/TextureEditor/Private/TextureEditorToolkit.cpp
@@ -1927,47 +1927,46 @@ TSharedRef<SWidget> FTextureEditorToolkit::MakeZoomControlWidget()
 
 	// zoom slider
 	TSharedRef<SWidget> ZoomControl = 
-		SNew(SHorizontalBox)
-		+ SHorizontalBox::Slot()
-		.AutoWidth()
-		.VAlign(VAlign_Center)
-		[
-			SNew(STextBlock)
-			.Text(LOCTEXT("ZoomLabel", "Zoom"))
-		]
-		+ SHorizontalBox::Slot()
-		.Padding(4.0f, 0.0f)
-		.AutoWidth()
-		.VAlign(VAlign_Center)
+		SNew(SBox)
+		.WidthOverride(250.f)
 		[
-			SNew(SBox)
-			.WidthOverride(200.f)
+			SNew(SHorizontalBox)
+			+ SHorizontalBox::Slot()
+			.AutoWidth()
+			.VAlign(VAlign_Center)
 			[
-				SNew(SSlider)
-				.OnValueChanged(this, &FTextureEditorToolkit::HandleZoomSliderChanged)
-				.Value(this, &FTextureEditorToolkit::HandleZoomSliderValue)
+				SNew(STextBlock)
+				.Text(LOCTEXT("ZoomLabel", "Zoom"))
 			]
-		]
-		+ SHorizontalBox::Slot()
-		.AutoWidth()
-		.VAlign(VAlign_Center)
-		.Padding(4.0f, 0.0f, 0.0f, 0.0f)
-		[
-			SNew(SComboButton)
-			.ComboButtonStyle(FAppStyle::Get(), "SimpleComboButton")
-			.ButtonContent()
+			+ SHorizontalBox::Slot()
+			.Padding(4.0f, 0.0f)
+			.VAlign(VAlign_Center)
 			[
 				SNew(SBox)
-				.WidthOverride(40.f)
-				.HAlign(HAlign_Right)
+				.WidthOverride(200.f)
 				[
-					SNew(STextBlock)
-					.Text(this, &FTextureEditorToolkit::HandleZoomPercentageText)
+					SNew(SSlider)
+					.OnValueChanged(this, &FTextureEditorToolkit::HandleZoomSliderChanged)
+					.Value(this, &FTextureEditorToolkit::HandleZoomSliderValue)
 				]
 			]
-			.MenuContent()
+			+ SHorizontalBox::Slot()
+			.VAlign(VAlign_Center)
+			.AutoWidth()
+			.Padding(4.0f, 0.0f, 0.0f, 0.0f)
 			[
-				ZoomMenuBuilder.MakeWidget()
+
+					SNew(SComboButton)
+					.ComboButtonStyle(FAppStyle::Get(), "SimpleComboButton")
+					.ButtonContent()
+					[
+						SNew(STextBlock)
+						.Text(this, &FTextureEditorToolkit::HandleZoomPercentageText)
+					]
+					.MenuContent()
+					[
+						ZoomMenuBuilder.MakeWidget()
+					]
 			]
 		];
 
