Fixed texture viewer not allowing to toggle off the Alpha channel of textures/render targets/etc.
* Added EPixelFormatChannelFlags to indicate which pixel format supports which channel
* Fixed Alpha channel not properly tagged as such for various pixel formats (e.g. BC5 being marked as "Has Alpha Channel")
* Fixed some pixel formats without an enum-to-string (and vice versa) conversion
* Fixed and unified computation of mip number, max mip levels and pixel format for all types of textures supported by the texture viewer

diff --git a/Engine/Source/Editor/TextureEditor/Private/TextureEditorToolkit.cpp b/Engine/Source/Editor/TextureEditor/Private/TextureEditorToolkit.cpp
--- a/Engine/Source/Editor/TextureEditor/Private/TextureEditorToolkit.cpp
+++ b/Engine/Source/Editor/TextureEditor/Private/TextureEditorToolkit.cpp
@@ -61,6 +61,34 @@ const FName FTextureEditorToolkit::PropertiesTabId(TEXT("TextureEditor_Propertie
 
 UNREALED_API void GetBestFitForNumberOfTiles(int32 InSize, int32& OutRatioX, int32& OutRatioY);
 
+EPixelFormatChannelFlags GetPixelFormatChannelFlagForButton(ETextureChannelButton InButton)
+{
+	switch (InButton)
+	{
+		case ETextureChannelButton::Red:
+		{
+			return EPixelFormatChannelFlags::R;
+		}
+		case ETextureChannelButton::Green:
+		{
+			return EPixelFormatChannelFlags::G;
+		}
+		case ETextureChannelButton::Blue:
+		{
+			return EPixelFormatChannelFlags::B;
+		}
+		case ETextureChannelButton::Alpha:
+		{
+			return EPixelFormatChannelFlags::A;
+		}
+		default:
+		{
+			check(false);
+		}
+	}
+
+	return EPixelFormatChannelFlags::None;
+}
 
 /* FTextureEditorToolkit structors
  *****************************************************************************/
@@ -151,7 +179,7 @@ void FTextureEditorToolkit::InitTextureEditor( const EToolkitMode::Type Mode, co
 
 	ExposureBias = 0;
 
-	bIsVolumeTexture = Texture->IsA<UVolumeTexture>() || Texture->IsA<UTextureRenderTargetVolume>();
+	bIsVolumeTexture = IsVolumeTexture();
 
 	switch (Texture->CompressionSettings)
 	{
@@ -439,18 +467,10 @@ void FTextureEditorToolkit::PopulateQuickInfo( )
 	}
 
 	UTexture2D* Texture2D = Cast<UTexture2D>(Texture);
-	UTextureRenderTarget2D* Texture2DRT = Cast<UTextureRenderTarget2D>(Texture);
-	UTextureCube* TextureCube = Cast<UTextureCube>(Texture);
-	UTexture2DArray* Texture2DArray = Cast<UTexture2DArray>(Texture);
-	UTextureRenderTarget2DArray* Texture2DArrayRT = Cast<UTextureRenderTarget2DArray>(Texture);
-	UTexture2DDynamic* Texture2DDynamic = Cast<UTexture2DDynamic>(Texture);
-	UVolumeTexture* VolumeTexture = Cast<UVolumeTexture>(Texture);
-	UTextureRenderTargetVolume* VolumeTextureRT = Cast<UTextureRenderTargetVolume>(Texture);
-	UMediaTexture* MediaTexture = Cast<UMediaTexture>(Texture);
 
-	const bool bIsVolume = VolumeTexture || VolumeTextureRT;
-	const bool bIsArray2D = Texture2DArray || Texture2DArrayRT;
-	const bool bIsCube = !!TextureCube;
+	const bool bIsVolume = IsVolumeTexture();
+	const bool bIsArray2D = Is2DArrayTexture();
+	const bool bIsCube = IsCubeTexture();
 
 	const uint32 SurfaceWidth = (uint32)Texture->GetSurfaceWidth();
 	const uint32 SurfaceHeight = (uint32)Texture->GetSurfaceHeight();
@@ -542,93 +562,19 @@ void FTextureEditorToolkit::PopulateQuickInfo( )
 	MethodText->SetText(FText::Format(NSLOCTEXT("TextureEditor", "QuickInfo_Method", "Method: {0}"), Method));
 	LODBiasText->SetText(FText::Format(NSLOCTEXT("TextureEditor", "QuickInfo_LODBias", "Combined LOD Bias: {0}"), FText::AsNumber(Texture->GetCachedLODBias())));
 
-	int32 TextureFormatIndex = PF_MAX;
-	
-	if (Texture2D)
-	{
-		TextureFormatIndex = Texture2D->GetPixelFormat(SpecifiedLayer);
-	}
-	else if (TextureCube)
-	{
-		TextureFormatIndex = TextureCube->GetPixelFormat();
-	}
-	else if (Texture2DArray) 
-	{
-		TextureFormatIndex = Texture2DArray->GetPixelFormat();
-	}
-	else if (Texture2DArrayRT)
-	{
-		TextureFormatIndex = Texture2DArrayRT->GetFormat();
-	}
-	else if (Texture2DRT)
-	{
-		TextureFormatIndex = Texture2DRT->GetFormat();
-	}
-	else if (Texture2DDynamic)
-	{
-		TextureFormatIndex = Texture2DDynamic->Format;
-	}
-	else if (VolumeTexture)
-	{
-		TextureFormatIndex = VolumeTexture->GetPixelFormat();
-	}
-	else if (VolumeTextureRT)
-	{
-		TextureFormatIndex = VolumeTextureRT->GetFormat();
-	}
-
-	if (TextureFormatIndex != PF_MAX)
+	EPixelFormat TextureFormat = GetPixelFormat();
+	if (TextureFormat != PF_MAX)
 	{
-		FormatText->SetText(FText::Format(NSLOCTEXT("TextureEditor", "QuickInfo_Format", "Format: {0}"), FText::FromString(GPixelFormats[TextureFormatIndex].Name)));
+		FormatText->SetText(FText::Format(NSLOCTEXT("TextureEditor", "QuickInfo_Format", "Format: {0}"), FText::FromString(GPixelFormats[(uint8)TextureFormat].Name)));
 	}
 
-	int32 NumMips = 1;
-	if (Texture2D)
-	{
-		NumMips = Texture2D->GetNumMips();
-	}
-	else if (TextureCube)
-	{
-		NumMips = TextureCube->GetNumMips();
-	}
-	else if (Texture2DArray) 
-	{
-		NumMips = Texture2DArray->GetNumMips();
-	}
-	else if (Texture2DArrayRT)
-	{
-		NumMips = Texture2DArrayRT->GetNumMips();
-	}
-	else if (Texture2DRT)
-	{
-		NumMips = Texture2DRT->GetNumMips();
-	}
-	else if (Texture2DDynamic)
-	{
-		NumMips = Texture2DDynamic->NumMips;
-	}
-	else if (VolumeTexture)
-	{
-		NumMips = VolumeTexture->GetNumMips();
-	}
-	else if (VolumeTextureRT)
-	{
-		NumMips = VolumeTextureRT->GetNumMips();
-	}
-	else if (MediaTexture)
-	{
-		NumMips = MediaTexture->GetTextureNumMips();
-	}
+	EPixelFormatChannelFlags ValidTextureChannels = GetPixelFormatValidChannels(TextureFormat);
+	HasAlphaChannelText->SetText(FText::Format(NSLOCTEXT("TextureEditor", "QuickInfo_HasAlphaChannel", "Has Alpha Channel: {0}"),
+		EnumHasAnyFlags(ValidTextureChannels, EPixelFormatChannelFlags::A) ? NSLOCTEXT("TextureEditor", "True", "True") : NSLOCTEXT("TextureEditor", "False", "False")));
+	HasAlphaChannelText->SetVisibility((ValidTextureChannels != EPixelFormatChannelFlags::None) ? EVisibility::Visible : EVisibility::Collapsed);
 
+	int32 NumMips = GetNumMips();
 	NumMipsText->SetText(FText::Format(NSLOCTEXT("TextureEditor", "QuickInfo_NumMips", "Number of Mips: {0}"), FText::AsNumber(NumMips)));
-
-	if (Texture2D)
-	{
-		HasAlphaChannelText->SetText(FText::Format(NSLOCTEXT("TextureEditor", "QuickInfo_HasAlphaChannel", "Has Alpha Channel: {0}"),
-			Texture2D->HasAlphaChannel() ? NSLOCTEXT("TextureEditor", "True", "True") : NSLOCTEXT("TextureEditor", "False", "False")));
-	}
-
-	HasAlphaChannelText->SetVisibility(Texture2D ? EVisibility::Visible : EVisibility::Collapsed);
 }
 
 
@@ -1088,60 +1034,109 @@ void FTextureEditorToolkit::FillToolbar(FToolBarBuilder& ToolbarBuilder)
 	}
 }
 
-TOptional<int32> FTextureEditorToolkit::GetMaxMipLevel( ) const
+TOptional<int32> FTextureEditorToolkit::GetMaxMipLevel() const
 {
-	const UTexture2D* Texture2D = Cast<UTexture2D>(Texture);
-	const UTextureCube* TextureCube = Cast<UTextureCube>(Texture);
-	const UTexture2DArray* Texture2DArray = Cast<UTexture2DArray>(Texture);
-	const UTextureRenderTargetCube* RTTextureCube = Cast<UTextureRenderTargetCube>(Texture);
-	const UTextureRenderTargetVolume* RTTextureVolume = Cast<UTextureRenderTargetVolume>(Texture);
-	const UTextureRenderTarget2D* RTTexture2D = Cast<UTextureRenderTarget2D>(Texture);
-	const UTextureRenderTarget2DArray* RTTexture2DArray = Cast<UTextureRenderTarget2DArray>(Texture);
-	const UVolumeTexture* VolumeTexture = Cast<UVolumeTexture>(Texture);
-
-	if (Texture2D)
+	TOptional<int32> MaxMipLevel;
+	int32 NumMips = GetNumMips();
+	if (NumMips > 0)
 	{
-		return Texture2D->GetNumMips() - 1;
+		MaxMipLevel = NumMips - 1;
 	}
+	return MaxMipLevel;
+}
 
-	if (TextureCube)
+int32 FTextureEditorToolkit::GetNumMips( ) const
+{
+	if (const UTexture2D* Texture2D = Cast<UTexture2D>(Texture))
 	{
-		return TextureCube->GetNumMips() - 1;
+		return Texture2D->GetNumMips();
 	}
-
-	if (Texture2DArray) 
+	else if (const UTextureCube* TextureCube = Cast<UTextureCube>(Texture))
 	{
-		return Texture2DArray->GetNumMips() - 1;
+		return TextureCube->GetNumMips();
 	}
-
-	if (RTTextureCube)
+	else if (const UTexture2DArray* Texture2DArray = Cast<UTexture2DArray>(Texture))
 	{
-		return RTTextureCube->GetNumMips() - 1;
+		return Texture2DArray->GetNumMips();
 	}
-
-	if (RTTextureVolume)
+	else if (const UVolumeTexture* VolumeTexture = Cast<UVolumeTexture>(Texture))
 	{
-		return RTTextureVolume->GetNumMips() - 1;
+		return VolumeTexture->GetNumMips();
 	}
-
-	if (RTTexture2D)
+	else if (const UTextureRenderTarget2D* Texture2DRT = Cast<UTextureRenderTarget2D>(Texture))
 	{
-		return RTTexture2D->GetNumMips() - 1;
+		return Texture2DRT->GetNumMips();
 	}
-
-	if (RTTexture2DArray)
+	else if (const UTextureRenderTargetCube* TextureCubeRT = Cast<UTextureRenderTargetCube>(Texture))
 	{
-		return RTTexture2DArray->GetNumMips() - 1;
+		return TextureCubeRT->GetNumMips();
 	}
-
-	if (VolumeTexture)
+	else if (const UTextureRenderTarget2DArray* Texture2DArrayRT = Cast<UTextureRenderTarget2DArray>(Texture))
+	{
+		return Texture2DArrayRT->GetNumMips();
+	}
+	else if (const UTextureRenderTargetVolume* VolumeTextureRT = Cast<UTextureRenderTargetVolume>(Texture))
+	{
+		return VolumeTextureRT->GetNumMips();
+	}
+	else if (const UTexture2DDynamic* Texture2DDynamic = Cast<UTexture2DDynamic>(Texture))
 	{
-		return VolumeTexture->GetNumMips() - 1;
+		return Texture2DDynamic->NumMips;
+	}
+	else if (const UMediaTexture* MediaTexture = Cast<UMediaTexture>(Texture))
+	{
+		return MediaTexture->GetTextureNumMips();
 	}
 
 	return MIPLEVEL_MAX;
 }
 
+EPixelFormat FTextureEditorToolkit::GetPixelFormat() const
+{
+	if (const UTexture2D* Texture2D = Cast<UTexture2D>(Texture))
+	{
+		return Texture2D->GetPixelFormat(SpecifiedLayer);
+	}
+	else if (const UTextureCube* TextureCube = Cast<UTextureCube>(Texture))
+	{
+		return TextureCube->GetPixelFormat();
+	}
+	else if (const UTexture2DArray* Texture2DArray = Cast<UTexture2DArray>(Texture))
+	{
+		return Texture2DArray->GetPixelFormat();
+	}
+	else if (const UVolumeTexture* VolumeTexture = Cast<UVolumeTexture>(Texture))
+	{
+		return VolumeTexture->GetPixelFormat();
+	}
+	else if (const UTextureRenderTarget2D* Texture2DRT = Cast<UTextureRenderTarget2D>(Texture))
+	{
+		return Texture2DRT->GetFormat();
+	}
+	else if (const UTextureRenderTargetCube* TextureCubeRT = Cast<UTextureRenderTargetCube>(Texture))
+	{
+		return TextureCubeRT->GetFormat();
+	}
+	else if (const UTextureRenderTarget2DArray* Texture2DArrayRT = Cast<UTextureRenderTarget2DArray>(Texture))
+	{
+		return Texture2DArrayRT->GetFormat();
+	}
+	else if (const UTextureRenderTargetVolume* VolumeTextureRT = Cast<UTextureRenderTargetVolume>(Texture))
+	{
+		return VolumeTextureRT->GetFormat();
+	}
+	else if (const UTexture2DDynamic* Texture2DDynamic = Cast<UTexture2DDynamic>(Texture))
+	{
+		return Texture2DDynamic->Format;
+	}
+	//else if (const UMediaTexture* MediaTexture = Cast<UMediaTexture>(Texture))
+	//{
+	//	UMediaTexture::GetDesc() suggests PF_B8G8R8A8, maybe?
+	//}
+
+	return PF_MAX;
+}
+
 TOptional<int32> FTextureEditorToolkit::GetMaxLayer() const
 {
 	return FMath::Max(Texture->Source.GetNumLayers() - 1, 1);
@@ -1195,18 +1190,10 @@ TSharedRef<SWidget> FTextureEditorToolkit::OnGenerateSettingsMenu()
 
 /* FTextureEditorToolkit callbacks
  *****************************************************************************/
-
-
-bool FTextureEditorToolkit::IsAlphaChannelButtonEnabled( ) const
+bool FTextureEditorToolkit::IsChannelButtonEnabled(ETextureChannelButton Button) const
 {
-	const UTexture2D* Texture2D = Cast<UTexture2D>(Texture);
-
-	if (Texture2D == NULL)
-	{
-		return false;
-	}
-
-	return Texture2D->HasAlphaChannel();
+	EPixelFormatChannelFlags ValidTextureChannels = GetPixelFormatValidChannels(GetPixelFormat());
+	return EnumHasAnyFlags(ValidTextureChannels, GetPixelFormatChannelFlagForButton(Button));
 }
 
 FSlateColor FTextureEditorToolkit::GetChannelButtonBackgroundColor(ETextureChannelButton Button) const
@@ -1681,6 +1668,7 @@ TSharedRef<SWidget> FTextureEditorToolkit::MakeChannelControlWidget()
 			.ForegroundColor(this, &FTextureEditorToolkit::GetChannelButtonForegroundColor, ETextureChannelButton::Red)
 			.OnCheckStateChanged_Lambda(OnChannelCheckStateChanged, ETextureChannelButton::Red)
 			.IsChecked(this, &FTextureEditorToolkit::OnGetChannelButtonCheckState, ETextureChannelButton::Red)
+			.IsEnabled(this, &FTextureEditorToolkit::IsChannelButtonEnabled, ETextureChannelButton::Red)
 			[
 				SNew(STextBlock)
 				.Font(FAppStyle::Get().GetFontStyle("TextureEditor.ChannelButtonFont"))
@@ -1698,6 +1686,7 @@ TSharedRef<SWidget> FTextureEditorToolkit::MakeChannelControlWidget()
 			.ForegroundColor(this, &FTextureEditorToolkit::GetChannelButtonForegroundColor, ETextureChannelButton::Green)
 			.OnCheckStateChanged_Lambda(OnChannelCheckStateChanged, ETextureChannelButton::Green)
 			.IsChecked(this, &FTextureEditorToolkit::OnGetChannelButtonCheckState, ETextureChannelButton::Green)
+			.IsEnabled(this, &FTextureEditorToolkit::IsChannelButtonEnabled, ETextureChannelButton::Green)
 			[
 				SNew(STextBlock)
 				.Font(FAppStyle::Get().GetFontStyle("TextureEditor.ChannelButtonFont"))
@@ -1716,6 +1705,7 @@ TSharedRef<SWidget> FTextureEditorToolkit::MakeChannelControlWidget()
 			.ForegroundColor(this, &FTextureEditorToolkit::GetChannelButtonForegroundColor, ETextureChannelButton::Blue)
 			.OnCheckStateChanged_Lambda(OnChannelCheckStateChanged, ETextureChannelButton::Blue)
 			.IsChecked(this, &FTextureEditorToolkit::OnGetChannelButtonCheckState, ETextureChannelButton::Blue)
+			.IsEnabled(this, &FTextureEditorToolkit::IsChannelButtonEnabled, ETextureChannelButton::Blue)
 			[
 				SNew(STextBlock)
 				.Font(FAppStyle::Get().GetFontStyle("TextureEditor.ChannelButtonFont"))
@@ -1733,7 +1723,7 @@ TSharedRef<SWidget> FTextureEditorToolkit::MakeChannelControlWidget()
 			.ForegroundColor(this, &FTextureEditorToolkit::GetChannelButtonForegroundColor, ETextureChannelButton::Alpha)
 			.OnCheckStateChanged_Lambda(OnChannelCheckStateChanged, ETextureChannelButton::Alpha)
 			.IsChecked(this, &FTextureEditorToolkit::OnGetChannelButtonCheckState, ETextureChannelButton::Alpha)
-			.IsEnabled(this, &FTextureEditorToolkit::IsAlphaChannelButtonEnabled)
+			.IsEnabled(this, &FTextureEditorToolkit::IsChannelButtonEnabled, ETextureChannelButton::Alpha)
 			[
 				SNew(STextBlock)
 				.Font(FAppStyle::Get().GetFontStyle("TextureEditor.ChannelButtonFont"))
diff --git a/Engine/Source/Editor/TextureEditor/Private/TextureEditorToolkit.h b/Engine/Source/Editor/TextureEditor/Private/TextureEditorToolkit.h
--- a/Engine/Source/Editor/TextureEditor/Private/TextureEditorToolkit.h
+++ b/Engine/Source/Editor/TextureEditor/Private/TextureEditorToolkit.h
@@ -129,6 +129,16 @@ protected:
 	 */
 	TOptional<int32> GetMaxMipLevel( ) const;
 
+	/**
+	 * Gets the number of mips that this texture supports.
+	 */
+	int32 GetNumMips() const;
+
+	/**
+	 * Gets the texture format of this texture.
+	 */
+	EPixelFormat GetPixelFormat() const;
+
 	TOptional<int32> GetMaxLayer() const;
 
 	/**
@@ -149,7 +159,7 @@ protected:
 	TSharedRef<SWidget> OnGenerateMipMapLevelMenu();
 	TSharedRef<SWidget> OnGenerateSettingsMenu();
 private:
-	bool IsAlphaChannelButtonEnabled() const;
+	bool IsChannelButtonEnabled(ETextureChannelButton Button) const;
 	FSlateColor GetChannelButtonBackgroundColor(ETextureChannelButton Button) const;
 	FSlateColor GetChannelButtonForegroundColor(ETextureChannelButton Button) const;
 	void OnChannelButtonCheckStateChanged(ETextureChannelButton Button);
diff --git a/Engine/Source/Runtime/Core/Public/PixelFormat.h b/Engine/Source/Runtime/Core/Public/PixelFormat.h
--- a/Engine/Source/Runtime/Core/Public/PixelFormat.h
+++ b/Engine/Source/Runtime/Core/Public/PixelFormat.h
@@ -145,7 +145,6 @@ enum EPixelFormat
 	op(PF_BC7) \
 	op(PF_R8_UINT) \
 	op(PF_L8) \
-	op(PF_R8) \
 	op(PF_XGXR8) \
 	op(PF_R8G8B8A8_UINT) \
 	op(PF_R8G8B8A8_SNORM) \
@@ -155,5 +154,22 @@ enum EPixelFormat
 	op(PF_PLATFORM_HDR_1) \
 	op(PF_PLATFORM_HDR_2) \
 	op(PF_NV12) \
-	op(PF_R32G32_UINT)
+	op(PF_R32G32_UINT) \
+	op(PF_ETC2_R11_EAC) \
+	op(PF_ETC2_RG11_EAC) \
+	op(PF_R8) 
 
+// Defines which channel is valid for each pixel format
+enum class EPixelFormatChannelFlags : uint8
+{
+	R = 1 << 0,
+	G = 1 << 1,
+	B = 1 << 2,
+	A = 1 << 3,
+	RG = R | G,
+	RGB = R | G | B,
+	RGBA = R | G | B | A,
+
+	None = 0,
+};
+ENUM_CLASS_FLAGS(EPixelFormatChannelFlags);
\ No newline at end of file
diff --git a/Engine/Source/Runtime/Engine/Private/Texture2D.cpp b/Engine/Source/Runtime/Engine/Private/Texture2D.cpp
--- a/Engine/Source/Runtime/Engine/Private/Texture2D.cpp
+++ b/Engine/Source/Runtime/Engine/Private/Texture2D.cpp
@@ -898,7 +898,7 @@ bool UTexture2D::HasAlphaChannel() const
 			return GetDefaultTexture2D(this)->HasAlphaChannel();
 		}
 #endif
-		return PrivatePlatformData->PixelFormat != PF_DXT1;
+		return EnumHasAnyFlags(GetPixelFormatValidChannels(PrivatePlatformData->PixelFormat), EPixelFormatChannelFlags::A);
 	}
 	return false;
 }
diff --git a/Engine/Source/Runtime/RenderCore/Private/RenderUtils.cpp b/Engine/Source/Runtime/RenderCore/Private/RenderUtils.cpp
--- a/Engine/Source/Runtime/RenderCore/Private/RenderUtils.cpp
+++ b/Engine/Source/Runtime/RenderCore/Private/RenderUtils.cpp
@@ -865,6 +865,86 @@ EPixelFormat GetPixelFormatFromString(const TCHAR* InPixelFormatStr)
 	return PF_Unknown;
 }
 
+EPixelFormatChannelFlags GetPixelFormatValidChannels(EPixelFormat InPixelFormat)
+{
+	static constexpr EPixelFormatChannelFlags PixelFormatToChannelFlags[] =
+	{
+		EPixelFormatChannelFlags::None,		// PF_Unknown,
+		EPixelFormatChannelFlags::RGBA,		// PF_A32B32G32R32F
+		EPixelFormatChannelFlags::RGBA,		// PF_B8G8R8A8
+		EPixelFormatChannelFlags::G,		// PF_G8
+		EPixelFormatChannelFlags::G,		// PF_G16
+		EPixelFormatChannelFlags::RGB,		// PF_DXT1
+		EPixelFormatChannelFlags::RGBA,		// PF_DXT3
+		EPixelFormatChannelFlags::RGBA,		// PF_DXT5
+		EPixelFormatChannelFlags::RG,		// PF_UYVY
+		EPixelFormatChannelFlags::RGB,		// PF_FloatRGB
+		EPixelFormatChannelFlags::RGBA,		// PF_FloatRGBA
+		EPixelFormatChannelFlags::None,		// PF_DepthStencil
+		EPixelFormatChannelFlags::None,		// PF_ShadowDepth
+		EPixelFormatChannelFlags::R,		// PF_R32_FLOAT
+		EPixelFormatChannelFlags::RG,		// PF_G16R16
+		EPixelFormatChannelFlags::RG,		// PF_G16R16F
+		EPixelFormatChannelFlags::RG,		// PF_G16R16F_FILTER
+		EPixelFormatChannelFlags::RG,		// PF_G32R32F
+		EPixelFormatChannelFlags::RGBA,		// PF_A2B10G10R10
+		EPixelFormatChannelFlags::RGBA,		// PF_A16B16G16R16
+		EPixelFormatChannelFlags::None,		// PF_D24
+		EPixelFormatChannelFlags::R,		// PF_R16F
+		EPixelFormatChannelFlags::R,		// PF_R16F_FILTER
+		EPixelFormatChannelFlags::RG,		// PF_BC5
+		EPixelFormatChannelFlags::RG,		// PF_V8U8
+		EPixelFormatChannelFlags::A,		// PF_A1
+		EPixelFormatChannelFlags::RGB,		// PF_FloatR11G11B10
+		EPixelFormatChannelFlags::A,		// PF_A8
+		EPixelFormatChannelFlags::R,		// PF_R32_UINT
+		EPixelFormatChannelFlags::RGBA,		// PF_R32_SINT
+		EPixelFormatChannelFlags::RGBA,		// PF_PVRTC2
+		EPixelFormatChannelFlags::RGBA,		// PF_PVRTC4
+		EPixelFormatChannelFlags::R,		// PF_R16_UINT
+		EPixelFormatChannelFlags::R,		// PF_R16_SINT
+		EPixelFormatChannelFlags::RGBA,		// PF_R16G16B16A16_UINT
+		EPixelFormatChannelFlags::RGBA,		// PF_R16G16B16A16_SINT
+		EPixelFormatChannelFlags::RGB,		// PF_R5G6B5_UNORM
+		EPixelFormatChannelFlags::RGBA,		// PF_R8G8B8A8
+		EPixelFormatChannelFlags::RGBA,		// PF_A8R8G8B8
+		EPixelFormatChannelFlags::R,		// PF_BC4
+		EPixelFormatChannelFlags::RG,		// PF_R8G8
+		EPixelFormatChannelFlags::RGB,		// PF_ATC_RGB
+		EPixelFormatChannelFlags::RGBA,		// PF_ATC_RGBA_E
+		EPixelFormatChannelFlags::RGBA,		// PF_ATC_RGBA_I
+		EPixelFormatChannelFlags::G,		// PF_X24_G8		
+		EPixelFormatChannelFlags::RGB,		// PF_ETC1
+		EPixelFormatChannelFlags::RGB,		// PF_ETC2_RGB
+		EPixelFormatChannelFlags::RGBA,		// PF_ETC2_RGBA
+		EPixelFormatChannelFlags::RGBA,		// PF_R32G32B32A32_UINT
+		EPixelFormatChannelFlags::RG,		// PF_R16G16_UINT
+		EPixelFormatChannelFlags::RGB,		// PF_ASTC_4x4
+		EPixelFormatChannelFlags::RGB,		// PF_ASTC_6x6
+		EPixelFormatChannelFlags::RGB,		// PF_ASTC_8x8
+		EPixelFormatChannelFlags::RGB,		// PF_ASTC_10x10
+		EPixelFormatChannelFlags::RGB,		// PF_ASTC_12x12
+		EPixelFormatChannelFlags::RGB,		// PF_BC6H
+		EPixelFormatChannelFlags::RGBA,		// PF_BC7
+		EPixelFormatChannelFlags::R,		// PF_R8_UINT
+		EPixelFormatChannelFlags::None,		// PF_L8
+		EPixelFormatChannelFlags::RGBA,		// PF_XGXR8
+		EPixelFormatChannelFlags::RGBA,		// PF_R8G8B8A8_UINT
+		EPixelFormatChannelFlags::RGBA,		// PF_R8G8B8A8_SNORM
+		EPixelFormatChannelFlags::RGBA,		// PF_R16G16B16A16_UNORM
+		EPixelFormatChannelFlags::RGBA,		// PF_R16G16B16A16_SNORM
+		EPixelFormatChannelFlags::RGBA,		// PF_PLATFORM_HDR_0
+		EPixelFormatChannelFlags::RGBA,		// PF_PLATFORM_HDR_1
+		EPixelFormatChannelFlags::RGBA,		// PF_PLATFORM_HDR_2
+		EPixelFormatChannelFlags::None,		// PF_NV12
+		EPixelFormatChannelFlags::RG,		// PF_R32G32_UINT
+		EPixelFormatChannelFlags::R,		// PF_ETC2_R11_EAC
+		EPixelFormatChannelFlags::RG,		// PF_ETC2_RG11_EAC
+		EPixelFormatChannelFlags::R,		// PF_R8
+	};
+	static_assert(UE_ARRAY_COUNT(PixelFormatToChannelFlags) == (uint8)PF_MAX, "Missing pixel format");
+	return (InPixelFormat < PF_MAX) ? PixelFormatToChannelFlags[(uint8)InPixelFormat] : EPixelFormatChannelFlags::None;
+}
 
 const TCHAR* GetCubeFaceName(ECubeFace Face)
 {
diff --git a/Engine/Source/Runtime/RenderCore/Public/RenderUtils.h b/Engine/Source/Runtime/RenderCore/Public/RenderUtils.h
--- a/Engine/Source/Runtime/RenderCore/Public/RenderUtils.h
+++ b/Engine/Source/Runtime/RenderCore/Public/RenderUtils.h
@@ -466,6 +466,14 @@ RENDERCORE_API const TCHAR* GetPixelFormatString(EPixelFormat InPixelFormat);
  */
 RENDERCORE_API EPixelFormat GetPixelFormatFromString(const TCHAR* InPixelFormatStr);
 
+/** 
+ *  Returns the valid channels for this pixel format
+ * 
+ * @return e.g. EPixelFormatChannelFlags::G for PF_G8
+ */
+RENDERCORE_API EPixelFormatChannelFlags GetPixelFormatValidChannels(EPixelFormat InPixelFormat);
+
+
 /**
  * Convert from ECubeFace to text string
  * @param Face - ECubeFace type to convert
