Undo changelist 15503762

#rnx

diff --git a/Engine/Source/Editor/TextureEditor/Private/TextureEditorToolkit.cpp b/Engine/Source/Editor/TextureEditor/Private/TextureEditorToolkit.cpp
--- a/Engine/Source/Editor/TextureEditor/Private/TextureEditorToolkit.cpp
+++ b/Engine/Source/Editor/TextureEditor/Private/TextureEditorToolkit.cpp
@@ -1083,10 +1083,10 @@ int32 FTextureEditorToolkit::GetNumMips( ) const
 	{
 		return Texture2DDynamic->NumMips;
 	}
-	//else if (const UMediaTexture* MediaTexture = Cast<UMediaTexture>(Texture))
-	//{
-	//	return MediaTexture->GetTextureNumMips();
-	//}
+	else if (const UMediaTexture* MediaTexture = Cast<UMediaTexture>(Texture))
+	{
+		return MediaTexture->GetTextureNumMips();
+	}
 
 	return MIPLEVEL_MAX;
 }
