Removed a reference to UE4 in a comment.

diff --git a/Engine/Source/Editor/TextureEditor/Private/Models/TextureEditorViewportClient.cpp b/Engine/Source/Editor/TextureEditor/Private/Models/TextureEditorViewportClient.cpp
--- a/Engine/Source/Editor/TextureEditor/Private/Models/TextureEditorViewportClient.cpp
+++ b/Engine/Source/Editor/TextureEditor/Private/Models/TextureEditorViewportClient.cpp
@@ -185,7 +185,7 @@ void FTextureEditorViewportClient::Draw(FViewport* Viewport, FCanvas* Canvas)
 
 		if (bIsVirtualTexture && Texture->Source.GetNumBlocks() > 1)
 		{
-			// Adjust UVs to display entire UDIM range, acounting for UE4 inverted V-axis
+			// Adjust UVs to display entire UDIM range, accounting for UE inverted V-axis
 			const FIntPoint BlockSize = Texture->Source.GetSizeInBlocks();
 			TileItem.UV0 = FVector2D(0.0f, 1.0f - (float)BlockSize.Y);
 			TileItem.UV1 = FVector2D((float)BlockSize.X, 1.0f);
