Add support for specifying Oodle encoding parameters in the texture editor for experimenting with different settings for individual textures.

diff --git a/Engine/Source/Editor/TextureEditor/Private/TextureEditorModule.cpp b/Engine/Source/Editor/TextureEditor/Private/TextureEditorModule.cpp
--- a/Engine/Source/Editor/TextureEditor/Private/TextureEditorModule.cpp
+++ b/Engine/Source/Editor/TextureEditor/Private/TextureEditorModule.cpp
@@ -37,6 +37,13 @@ public:
 		TSharedRef<FTextureEditorToolkit> NewTextureEditor(new FTextureEditorToolkit());
 		NewTextureEditor->InitTextureEditor(Mode, InitToolkitHost, Texture);
 
+		// OpenAssetEditor should ensure that we never have two toolkits open for the same asset!
+		check(Texture->TextureEditorCustomEncoding.IsValid() == false);
+
+		// We save this as a separate object so that the engine can reference the type
+		// without needing the texture editor module.
+		Texture->TextureEditorCustomEncoding = NewTextureEditor->CustomEncoding;
+
 		return NewTextureEditor;
 	}
 
diff --git a/Engine/Source/Editor/TextureEditor/Private/TextureEditorToolkit.cpp b/Engine/Source/Editor/TextureEditor/Private/TextureEditorToolkit.cpp
--- a/Engine/Source/Editor/TextureEditor/Private/TextureEditorToolkit.cpp
+++ b/Engine/Source/Editor/TextureEditor/Private/TextureEditorToolkit.cpp
@@ -48,6 +48,8 @@
 #include "Menus/TextureEditorViewOptionsMenu.h"
 #include "MediaTexture.h"
 #include "TextureEncodingSettings.h"
+#include "EditorWidgets/Public/SEnumCombo.h"
+#include "Widgets/Layout/SHeader.h"
 
 #define LOCTEXT_NAMESPACE "FTextureEditorToolkit"
 
@@ -94,6 +96,18 @@ EPixelFormatChannelFlags GetPixelFormatChannelFlagForButton(ETextureChannelButto
 	return EPixelFormatChannelFlags::None;
 }
 
+static void PostTextureRecode(UTexture* Texture)
+{
+	// Each time we change a custom encode setting we want to re-encode the texture
+	// as though we changed a compression setting on the actual texture, so we just
+	// post a CompressionSettings property changed event to handle all of that for
+	// us.
+	FProperty* Property = FindFProperty<FProperty>(UTexture::StaticClass(), "CompressionSettings");
+	FPropertyChangedEvent PropertyChangedEvent(Property);
+	Texture->PostEditChangeProperty(PropertyChangedEvent);
+}
+
+
 /* FTextureEditorToolkit structors
  *****************************************************************************/
 
@@ -119,6 +133,13 @@ FTextureEditorToolkit::~FTextureEditorToolkit( )
 	GEditor->GetEditorSubsystem<UImportSubsystem>()->OnAssetPostImport.RemoveAll(this);
 
 	GEditor->UnregisterForUndo(this);
+
+	if (CustomEncoding->bUseCustomEncode)
+	{
+		// reencode the texture with normal settings.
+		CustomEncoding->bUseCustomEncode = false;
+		PostTextureRecode(Texture);
+	}
 }
 
 
@@ -181,6 +202,8 @@ void FTextureEditorToolkit::InitTextureEditor( const EToolkitMode::Type Mode, co
 	Texture->SetFlags(RF_Transactional);
 	GEditor->RegisterForUndo(this);
 
+	CustomEncoding = MakeShared<FTextureEditorCustomEncode>(FTextureEditorCustomEncode());
+
 	// initialize view options
 	bIsRedChannel = true;
 	bIsGreenChannel = true;
@@ -544,16 +567,24 @@ void FTextureEditorToolkit::PopulateQuickInfo( )
 				OodleEffortLabel->SetVisibility(EVisibility::Visible);
 				OodleTilingLabel->SetVisibility(EVisibility::Visible);
 
-				EncodeSpeedText->SetText(
-					ResultMetadata.EncodeSpeed == (uint8)ETextureEncodeSpeed::Fast ?
-						NSLOCTEXT("TextureEditor", "QuickInfo_EncodeSpeed_Fast", "Encode Speed: Fast") : 
-						NSLOCTEXT("TextureEditor", "QuickInfo_EncodeSpeed_Final", "Encode Speed: Final")
-				);
-				OodleEncodeSpeedText->SetText(
-					ResultMetadata.EncodeSpeed == (uint8)ETextureEncodeSpeed::Fast ?
-					FText::AsCultureInvariant("Fast") :
-					FText::AsCultureInvariant("Final")
-				);			
+				if (ResultMetadata.bWasEditorCustomEncoding)
+				{
+					EncodeSpeedText->SetText(NSLOCTEXT("TextureEditor", "QuickInfo_EncodeSpeed_Custom", "Encode Speed: Custom"));
+					OodleEncodeSpeedText->SetText(NSLOCTEXT("TextureEditor", "QuickInfoDetails_EncodeSpeed_Custom", "Custom"));
+				}
+				else
+				{
+					EncodeSpeedText->SetText(
+						ResultMetadata.EncodeSpeed == (uint8)ETextureEncodeSpeed::Fast ?
+							NSLOCTEXT("TextureEditor", "QuickInfo_EncodeSpeed_Fast", "Encode Speed: Fast") : 
+							NSLOCTEXT("TextureEditor", "QuickInfo_EncodeSpeed_Final", "Encode Speed: Final")
+					);
+					OodleEncodeSpeedText->SetText(
+						ResultMetadata.EncodeSpeed == (uint8)ETextureEncodeSpeed::Fast ?
+						NSLOCTEXT("TextureEditor", "QuickInfoDetails_EncodeSpeed_Fast", "Fast") :
+						NSLOCTEXT("TextureEditor", "QuickInfoDetails_EncodeSpeed_Final", "Final")
+					);
+				}
 
 				if (ResultMetadata.OodleRDO == 0)
 				{
@@ -561,7 +592,11 @@ void FTextureEditorToolkit::PopulateQuickInfo( )
 					const bool bDisabledGlobally = ResultMetadata.EncodeSpeed == (uint8)ETextureEncodeSpeed::Fast ? !Settings->bFastUsesRDO : !Settings->bFinalUsesRDO;
 
 					OodleRDOText->SetText(NSLOCTEXT("TextureEditor", "QuickInfo_Oodle_RDODisable", "Disabled"));
-					if (bDisabledGlobally)
+					if (ResultMetadata.bWasEditorCustomEncoding)
+					{
+						OodleRDOSourceText->SetText(NSLOCTEXT("TextureEditor", "QuickInfo_Oodle_RDOSource_Custom", "Custom"));
+					}
+					else if (bDisabledGlobally)
 					{
 						OodleRDOSourceText->SetText(NSLOCTEXT("TextureEditor", "QuickInfo_Oodle_RDOSourceDisableSettings", "Disabled By Project Settings"));
 					}
@@ -585,7 +620,11 @@ void FTextureEditorToolkit::PopulateQuickInfo( )
 				{
 					OodleRDOText->SetText(FText::AsNumber(ResultMetadata.OodleRDO));
 
-					if (ResultMetadata.RDOSource == FTexturePlatformData::FTextureEncodeResultMetadata::OodleRDOSource::Default)
+					if (ResultMetadata.bWasEditorCustomEncoding)
+					{
+						OodleRDOSourceText->SetText(NSLOCTEXT("TextureEditor", "QuickInfo_Oodle_RDOSource_Custom", "Custom"));
+					}
+					else if (ResultMetadata.RDOSource == FTexturePlatformData::FTextureEncodeResultMetadata::OodleRDOSource::Default)
 					{
 						OodleRDOSourceText->SetText(NSLOCTEXT("TextureEditor", "QuickInfo_Oodle_RDOSource_Default", "Project (Lambda)"));
 					}
@@ -1113,7 +1152,125 @@ void FTextureEditorToolkit::CreateInternalWidgets( )
 					SAssignNew(OodleTilingText, STextBlock)
 				]
 		]
-	];
+	]
+	+ SVerticalBox::Slot()
+		.AutoHeight()
+		.VAlign(VAlign_Center)
+		[
+			SNew(SHeader)
+				.HAlign(EHorizontalAlignment::HAlign_Fill)
+			[
+				SNew(STextBlock)
+				.Text(LOCTEXT("OodleTab_Label_TryHeader", "Try Encodings"))
+			]
+		]
+
+	+ SVerticalBox::Slot()
+		.AutoHeight()
+		.Padding(4.0f)
+		[
+			SNew(SHorizontalBox)
+			+ SHorizontalBox::Slot()
+				.FillWidth(0.5f)
+				[
+					SNew(SVerticalBox)						
+
+						+ SVerticalBox::Slot()
+							.AutoHeight()
+							.VAlign(VAlign_Center)
+							.Padding(6)
+							[
+								SNew(STextBlock)
+								.Text(LOCTEXT("OodleTab_Label_OverrideCompression", "Enabled:"))
+							.ToolTipText(LOCTEXT("OodleTab_ToolTip_OverrideCompression", "If checked, allows you to experiment with Oodle RDO compression settings to visualize results."))
+
+							]
+						+ SVerticalBox::Slot()
+							.AutoHeight()
+							.VAlign(VAlign_Center)
+							.Padding(6)
+							[
+								SNew(STextBlock)
+								.Text(LOCTEXT("OodleTab_Label_OverrideRDO", "RDO Lambda:"))
+							.ToolTipText(LOCTEXT("OodleTab_ToolTip_OverrideRDO", "The RDO lambda to encode with for experimentation. 0 disables RDO entirely. 1 is largest filesize, 100 is smallest."))
+
+							]
+						+ SVerticalBox::Slot()
+							.AutoHeight()
+							.VAlign(VAlign_Center)
+							.Padding(6)
+							[
+								SNew(STextBlock)
+								.Text(LOCTEXT("OodleTab_Label_OverrideEffort", "Effort:"))
+							.ToolTipText(LOCTEXT("OodleTab_ToolTip_OverrideEffort", "The encoding effort to try. Effort controls how much CPU time spent on finding better results. See the Oodle Texture documentation for detailed information."))
+
+							]
+						+ SVerticalBox::Slot()
+							.AutoHeight()
+							.VAlign(VAlign_Center)
+							.Padding(6)
+							[
+								SNew(STextBlock)
+								.Text(LOCTEXT("OodleTab_Label_OverrideTiling", "Universal Tiling:"))
+							.ToolTipText(LOCTEXT("OodleTab_ToolTip_OverrideTiling", "The universal tiling to try. See the Oodle Texture documentation for detailed information."))
+
+							]
+				]
+
+			+ SHorizontalBox::Slot()
+				.FillWidth(0.5f)
+				[
+					SNew(SVerticalBox)
+
+
+						+ SVerticalBox::Slot()
+							.AutoHeight()
+							.VAlign(VAlign_Center)
+							.Padding(2)
+							[
+								SAssignNew(OodleOverrideCheck, SCheckBox)
+								.OnCheckStateChanged(this, &FTextureEditorToolkit::OnUseEditorOodleSettingsChanged)
+								.IsChecked(this, &FTextureEditorToolkit::UseEditorOodleSettingsChecked)
+							]
+
+						+ SVerticalBox::Slot()
+							.AutoHeight()
+							.VAlign(VAlign_Center)
+							.Padding(2)
+							[
+								SNew(SNumericEntryBox<int32>)
+								.Value(this, &FTextureEditorToolkit::GetEditorOodleSettingsRDO)
+								.OnValueCommitted(this, &FTextureEditorToolkit::EditorOodleSettingsRDOCommitted)
+								.IsEnabled(this, &FTextureEditorToolkit::EditorOodleSettingsEnabled)
+							]
+
+						+ SVerticalBox::Slot()
+							.AutoHeight()
+							.VAlign(VAlign_Center)
+							.Padding(2)
+							[
+								SNew(SEnumComboBox, StaticEnum< ETextureEncodeEffort >())
+								.CurrentValue(this, &FTextureEditorToolkit::GetEditorOodleSettingsEffort)
+								.OnEnumSelectionChanged(this, &FTextureEditorToolkit::EditorOodleSettingsEffortChanged)
+								.IsEnabled(this, &FTextureEditorToolkit::EditorOodleSettingsEnabled)
+							]
+
+						+ SVerticalBox::Slot()
+							.AutoHeight()
+							.VAlign(VAlign_Center)
+							.Padding(2)
+							[
+								SNew(SEnumComboBox, StaticEnum< ETextureUniversalTiling >())
+								.CurrentValue(this, &FTextureEditorToolkit::GetEditorOodleSettingsTiling)
+								.OnEnumSelectionChanged(this, &FTextureEditorToolkit::EditorOodleSettingsTilingChanged)
+								.IsEnabled(this, &FTextureEditorToolkit::EditorOodleSettingsEnabled)
+							]
+				]
+		];
+
+
+
+
 
 	TextureProperties = SNew(SVerticalBox)
 
@@ -1951,6 +2108,80 @@ float FTextureEditorToolkit::HandleZoomSliderValue() const
 	return (CalculateDisplayedZoomLevel() / MaxZoom);
 }
 
+int32 FTextureEditorToolkit::GetEditorOodleSettingsEffort() const
+{
+	return CustomEncoding->OodleEncodeEffort;
+}
+
+void FTextureEditorToolkit::EditorOodleSettingsEffortChanged(int32 NewValue, ESelectInfo::Type SelectionType)
+{
+	CustomEncoding->OodleEncodeEffort = NewValue;
+
+	if (CustomEncoding->bUseCustomEncode)
+	{
+		PostTextureRecode(Texture);
+	}
+}
+
+int32 FTextureEditorToolkit::GetEditorOodleSettingsTiling() const
+{
+	return CustomEncoding->OodleUniversalTiling;
+}
+
+void FTextureEditorToolkit::EditorOodleSettingsTilingChanged(int32 NewValue, ESelectInfo::Type SelectionType)
+{
+	CustomEncoding->OodleUniversalTiling = NewValue;
+
+	if (CustomEncoding->bUseCustomEncode)
+	{
+		PostTextureRecode(Texture);
+	}
+}
+
+TOptional<int32> FTextureEditorToolkit::GetEditorOodleSettingsRDO() const
+{
+	return CustomEncoding->OodleRDOLambda;
+}
+
+void FTextureEditorToolkit::EditorOodleSettingsRDOCommitted(int32 NewValue, ETextCommit::Type CommitType)
+{
+	if (NewValue > 100)
+	{
+		NewValue = 100;
+	}
+	if (NewValue < 0)
+	{
+		NewValue = 0;
+	}
+
+	CustomEncoding->OodleRDOLambda = (int8)NewValue;
+
+	if (CustomEncoding->bUseCustomEncode)
+	{
+		PostTextureRecode(Texture);
+	}
+}
+
+
+bool FTextureEditorToolkit::EditorOodleSettingsEnabled() const
+{
+	return CustomEncoding->bUseCustomEncode;
+}
+
+ECheckBoxState FTextureEditorToolkit::UseEditorOodleSettingsChecked() const
+{
+	return CustomEncoding->bUseCustomEncode ? ECheckBoxState::Checked : ECheckBoxState::Unchecked;
+}
+
+void FTextureEditorToolkit::OnUseEditorOodleSettingsChanged(ECheckBoxState NewState)
+{
+	// We need to convince the texture to recompress and signal all its users
+	// that they need to update, so we do this by faking a compression method property change.
+	CustomEncoding->bUseCustomEncode = NewState == ECheckBoxState::Checked ? true : false;
+
+	PostTextureRecode(Texture);
+}
+
 TSharedRef<SWidget> FTextureEditorToolkit::MakeChannelControlWidget()
 {
 	auto OnChannelCheckStateChanged = [this](ECheckBoxState NewState, ETextureChannelButton Button)
diff --git a/Engine/Source/Editor/TextureEditor/Private/TextureEditorToolkit.h b/Engine/Source/Editor/TextureEditor/Private/TextureEditorToolkit.h
--- a/Engine/Source/Editor/TextureEditor/Private/TextureEditorToolkit.h
+++ b/Engine/Source/Editor/TextureEditor/Private/TextureEditorToolkit.h
@@ -403,6 +403,28 @@ private:
 	FRotator VolumeOrientation;
 
 	bool bIsVolumeTexture;
+
+	// Objects and callbacks for the custom encoding settings checkbox
+	TSharedPtr<class SCheckBox> OodleOverrideCheck;
+	void OnUseEditorOodleSettingsChanged(ECheckBoxState NewState);
+	ECheckBoxState UseEditorOodleSettingsChecked() const;
+
+	// callback for enabling custom encoding sub-controls
+	bool EditorOodleSettingsEnabled() const;
+
+	// accessor callbacks for custom encoding sub controls
+	int32 GetEditorOodleSettingsEffort() const;
+	void EditorOodleSettingsEffortChanged(int32 NewValue, ESelectInfo::Type SelectionType);
+	TOptional<int32> GetEditorOodleSettingsRDO() const;
+	void EditorOodleSettingsRDOCommitted(int32 NewValue, ETextCommit::Type CommitType);
+	int32 GetEditorOodleSettingsTiling() const;
+	void EditorOodleSettingsTilingChanged(int32 NewValue, ESelectInfo::Type SelectionType);
+
+public:
+	// These are on a separate object so then engine can share the type without needing
+	// our module.
+	TSharedPtr<class FTextureEditorCustomEncode> CustomEncoding;
+
 private:
 
 	// Tab names
diff --git a/Engine/Source/Editor/TextureEditor/TextureEditor.Build.cs b/Engine/Source/Editor/TextureEditor/TextureEditor.Build.cs
--- a/Engine/Source/Editor/TextureEditor/TextureEditor.Build.cs
+++ b/Engine/Source/Editor/TextureEditor/TextureEditor.Build.cs
@@ -46,6 +46,7 @@ public class TextureEditor : ModuleRules
 				"UnrealEd",
                 "PropertyEditor",
 				"MediaUtils",
+				"EditorWidgets",
 				"MediaAssets"
 			}
 		);
diff --git a/Engine/Source/Runtime/Engine/Classes/Engine/Texture.h b/Engine/Source/Runtime/Engine/Classes/Engine/Texture.h
--- a/Engine/Source/Runtime/Engine/Classes/Engine/Texture.h
+++ b/Engine/Source/Runtime/Engine/Classes/Engine/Texture.h
@@ -680,6 +680,11 @@ struct FTexturePlatformData
 		// the various Oodle specific values right now.)
 		bool bSupportsEncodeSpeed = false;
 
+		// If true, then the encoding settings were overridden in the texture editor
+		// for encoding experimentation, and thus RDOSource and EncodeSpeed should 
+		// be ignored.
+		bool bWasEditorCustomEncoding = false;
+
 		enum class OodleRDOSource : uint8
 		{
 			Default,	// We defaulted back to the project settings
@@ -1178,6 +1183,9 @@ public:
 	/** Indicates ImportCustomProperties has been called (set in ImportCustomProperties, unset in the subsequent PostEditChange) */
 	uint8 bCustomPropertiesImported : 1;
 
+	// When we are open in an asset editor, we have a pointer to a custom encoding
+	// object which can optionally cause us to do something other than Fast/Final encode settings.
+	TWeakPtr<class FTextureEditorCustomEncode> TextureEditorCustomEncoding;
 #endif // WITH_EDITORONLY_DATA
 
 	/** If true, the RHI texture will be created using TexCreate_NoTiling */
diff --git a/Engine/Source/Runtime/Engine/Private/TextureDerivedData.cpp b/Engine/Source/Runtime/Engine/Private/TextureDerivedData.cpp
--- a/Engine/Source/Runtime/Engine/Private/TextureDerivedData.cpp
+++ b/Engine/Source/Runtime/Engine/Private/TextureDerivedData.cpp
@@ -2224,6 +2224,38 @@ void UTexture::CachePlatformData(bool bAsyncCache, bool bAllowAsyncBuild, bool b
 				GetBuildSettingsForRunningPlatform(*this, ETextureEncodeSpeed::Final, BuildSettingsFetchOrBuild, &ResultMetadataFetchOrBuild);
 			}
 
+			// If we're open in a texture editor, then we might have custom build settings.
+			if (TextureEditorCustomEncoding.IsValid())
+			{
+				TSharedPtr<FTextureEditorCustomEncode> CustomEncoding = TextureEditorCustomEncoding.Pin();
+				if (CustomEncoding.IsValid() && // (threading) could have been destroyed between weak ptr IsValid and Pin
+					CustomEncoding->bUseCustomEncode)
+				{
+					// If we are overriding, we don't want to have a fetch first, so just set our encode
+					// speed to whatever we already have staged, then set those settings to the custom
+					// ones.
+					EncodeSpeed = (ETextureEncodeSpeed)BuildSettingsFetchOrBuild[0].RepresentsEncodeSpeedNoSend;
+
+					for (int32 i = 0; i < BuildSettingsFetchOrBuild.Num(); i++)
+					{
+						FTextureBuildSettings& BuildSettings = BuildSettingsFetchOrBuild[i];
+						FTexturePlatformData::FTextureEncodeResultMetadata& ResultMetadata = ResultMetadataFetchOrBuild[i];
+
+						BuildSettings.OodleRDO = CustomEncoding->OodleRDOLambda;
+						BuildSettings.bOodleUsesRDO = CustomEncoding->OodleRDOLambda ? true : false;
+						BuildSettings.OodleEncodeEffort = CustomEncoding->OodleEncodeEffort;
+						BuildSettings.OodleUniversalTiling = CustomEncoding->OodleUniversalTiling;
+
+						ResultMetadata.OodleRDO = CustomEncoding->OodleRDOLambda;
+						ResultMetadata.OodleEncodeEffort = CustomEncoding->OodleEncodeEffort;
+						ResultMetadata.OodleUniversalTiling = CustomEncoding->OodleUniversalTiling;
+						ResultMetadata.EncodeSpeed = (uint8)EncodeSpeed;
+
+						ResultMetadata.bWasEditorCustomEncoding = true;
+					}
+				}
+			}
+
 			check(BuildSettingsFetchOrBuild.Num() == Source.GetNumLayers());
 
 			// The only time we don't cache is if we a) have existing data and b) it matches what we want.
diff --git a/Engine/Source/Runtime/Engine/Public/TextureEncodingSettings.h b/Engine/Source/Runtime/Engine/Public/TextureEncodingSettings.h
--- a/Engine/Source/Runtime/Engine/Public/TextureEncodingSettings.h
+++ b/Engine/Source/Runtime/Engine/Public/TextureEncodingSettings.h
@@ -160,3 +160,24 @@ class ENGINE_API UTextureEncodingUserSettings : public UDeveloperSettings
 	UPROPERTY(config, EditAnywhere, Category = EncodeSpeeds, meta = (ConfigRestartRequired = true))
 	ETextureEncodeSpeedOverride ForceEncodeSpeed;
 };
+
+//
+// Separate type so that the engine can check for custom encoding set in the texture
+// editor module without needing to depend on it.
+//
+class FTextureEditorCustomEncode
+{
+public:
+	// If we want to override Oodle specific encoding settings, we set this to true.
+	bool bUseCustomEncode = false;
+
+	// [0,100]
+	uint8 OodleRDOLambda = 0;
+
+	// enum ETextureEncodeEffort
+	uint8 OodleEncodeEffort = 0;
+
+	// enum ETextureUniversalTiling
+	uint8 OodleUniversalTiling = 0;
+
+};
\ No newline at end of file
