Added a way to view an estimate of benefits of RDO compression to the texture editor.

diff --git a/Engine/Source/Editor/TextureEditor/Private/TextureEditorToolkit.cpp b/Engine/Source/Editor/TextureEditor/Private/TextureEditorToolkit.cpp
--- a/Engine/Source/Editor/TextureEditor/Private/TextureEditorToolkit.cpp
+++ b/Engine/Source/Editor/TextureEditor/Private/TextureEditorToolkit.cpp
@@ -44,12 +44,16 @@
 #include "TextureEditorSettings.h"
 #include "TextureCompiler.h"
 #include "Widgets/Input/SSlider.h"
+#include "Widgets/Input/STextComboBox.h"
 #include "Widgets/Layout/SSpacer.h"
 #include "Menus/TextureEditorViewOptionsMenu.h"
 #include "MediaTexture.h"
 #include "TextureEncodingSettings.h"
 #include "EditorWidgets/Public/SEnumCombo.h"
 #include "Widgets/Layout/SHeader.h"
+#include "DerivedDataCache/Public/DerivedDataCacheKey.h"
+#include "Settings/ProjectPackagingSettings.h"
+#include "Compression/OodleDataCompressionUtil.h"
 
 #define LOCTEXT_NAMESPACE "FTextureEditorToolkit"
 
@@ -96,7 +100,7 @@ EPixelFormatChannelFlags GetPixelFormatChannelFlagForButton(ETextureChannelButto
 	return EPixelFormatChannelFlags::None;
 }
 
-static void PostTextureRecode(UTexture* Texture)
+void FTextureEditorToolkit::PostTextureRecode()
 {
 	// Each time we change a custom encode setting we want to re-encode the texture
 	// as though we changed a compression setting on the actual texture, so we just
@@ -105,6 +109,9 @@ static void PostTextureRecode(UTexture* Texture)
 	FProperty* Property = FindFProperty<FProperty>(UTexture::StaticClass(), "CompressionSettings");
 	FPropertyChangedEvent PropertyChangedEvent(Property);
 	Texture->PostEditChangeProperty(PropertyChangedEvent);
+
+	// Clear the key we have so we know when we have new data
+	OodleCompressedPreviewDDCKey.Set<FString>(FString());
 }
 
 
@@ -138,7 +145,7 @@ FTextureEditorToolkit::~FTextureEditorToolkit( )
 	{
 		// reencode the texture with normal settings.
 		CustomEncoding->bUseCustomEncode = false;
-		PostTextureRecode(Texture);
+		PostTextureRecode();
 	}
 }
 
@@ -531,6 +538,60 @@ void FTextureEditorToolkit::PopulateQuickInfo( )
 		}
 		else
 		{
+			//
+			// Check if we need to compress new Oodle preview once we know we have
+			// valid results.
+			//
+			FTexturePlatformData* PlatformData = PlatformDataPtr[0];
+
+			bool AlreadyHaveResults = false;
+			if (PlatformData->DerivedDataKey.GetIndex() == OodleCompressedPreviewDDCKey.GetIndex())
+			{
+				if (PlatformData->DerivedDataKey.IsType<FString>())
+				{
+					if (PlatformData->DerivedDataKey.Get<FString>() == OodleCompressedPreviewDDCKey.Get<FString>())
+					{
+						AlreadyHaveResults = true;
+					}
+				}
+				else
+				{
+					if (*PlatformData->DerivedDataKey.Get<UE::DerivedData::FCacheKeyProxy>().AsCacheKey() == *OodleCompressedPreviewDDCKey.Get<UE::DerivedData::FCacheKeyProxy>().AsCacheKey())
+					{
+						AlreadyHaveResults = true;
+					}
+				}
+			}
+
+			if (AlreadyHaveResults == false)
+			{
+				if (bEstimateCompressionEnabled)
+				{
+					OutstandingEstimation = PlatformData->LaunchEstimateOnDiskSizeTask(OodleCompressor, OodleCompressionLevel, CompressionBlockSize, Texture->GetPathName());
+				}
+
+				OodleCompressedPreviewDDCKey = PlatformDataPtr[0]->DerivedDataKey;
+			}
+
+			// If we have an outstanding estimation task, update UI when complete.
+			if (OutstandingEstimation.IsValid())
+			{
+				if (OutstandingEstimation.IsReady())
+				{
+					TTuple<uint64, uint64> Result = OutstandingEstimation.Get();
+
+					OodleEstimateRaw->SetText(FText::AsMemory(Result.Get<1>()));
+					OodleEstimateCompressed->SetText(FText::AsMemory(Result.Get<0>()));
+
+					OutstandingEstimation = TFuture<TTuple<uint64, uint64>>();
+				}
+				else
+				{
+					OodleEstimateRaw->SetText(NSLOCTEXT("TextureEditor", "QuickInfo_Oodle_Working", "Working..."));
+					OodleEstimateCompressed->SetText(NSLOCTEXT("TextureEditor", "QuickInfo_Oodle_Working", "Working..."));
+				}
+			}
+
 			OodleEncoderText->SetText(FText::FromName(ResultMetadata.Encoder));
 
 			if (ResultMetadata.bSupportsEncodeSpeed == false)
@@ -1023,253 +1084,481 @@ TSharedRef<SWidget> FTextureEditorToolkit::BuildTexturePropertiesWidget( )
 	return TexturePropertiesWidget.ToSharedRef();
 }
 
-void FTextureEditorToolkit::CreateInternalWidgets( )
+void FTextureEditorToolkit::CreateInternalWidgets()
 {
-	TextureViewport = SNew(STextureEditorViewport, SharedThis(this));
+	//
+	// Convert the packaging settings names into enums we can use.
+	//
+	UProjectPackagingSettings const* ProjectSettings = GetDefault<UProjectPackagingSettings>();
+	
+	PackagingSettingsNames.Add(MakeShared<FString>(TEXT("DebugDevelopment")));
+	PackagingSettingsNames.Add(MakeShared<FString>(TEXT("TestShipping")));
+	PackagingSettingsNames.Add(MakeShared<FString>(TEXT("Distribution")));
+	
+	// Default to Distribution
+	TSharedPtr<FString> InitialPackagingSetting = PackagingSettingsNames[2];
 
-	OodleTabContainer = SNew(SVerticalBox)
-	+ SVerticalBox::Slot()
-	.AutoHeight()
-	.Padding(4.0f)
-	[
-		SNew(SHorizontalBox)
-		+ SHorizontalBox::Slot()
-		.FillWidth(0.5f)
-		[
-			SNew(SVerticalBox)
+	// Determine which oodle encoder they are using.
+	const TCHAR* CompressorName = 0;
+	{
+		// Validity check the string by trying to convert to enum.
+		const FString& LookupCompressor = ProjectSettings->PackageCompressionMethod;
+		FOodleDataCompression::ECompressor PackageCompressor = FOodleDataCompression::ECompressor::Kraken;
+		if (FOodleDataCompression::ECompressorFromString(LookupCompressor, PackageCompressor))
+		{
+			OodleCompressor = PackageCompressor;
+		}
+		else
+		{
+			UE_LOG(LogTextureEditor, Warning, TEXT("Project packaging settings not using Oodle? Failed to recognize compression: %s - using Kraken for estimation."), *LookupCompressor);
+			OodleCompressor = FOodleDataCompression::ECompressor::Kraken;
+		}
 
-			+ SVerticalBox::Slot()
-				.AutoHeight()
-				.VAlign(VAlign_Center)
-				.Padding(4.0f)
-				[
-					SNew(STextBlock)
-						.Text(LOCTEXT("OodleTab_Label_Encoder", "Encoder:"))
-						.ToolTipText(LOCTEXT("OodleTab_Tooltip_Encoder", "Which texture encoder was used to encode the texture."))
-				]
+		FOodleDataCompression::ECompressorToString(OodleCompressor, &CompressorName);
+	}
 
-			+ SVerticalBox::Slot()
-				.AutoHeight()
-				.VAlign(VAlign_Center)
-				.Padding(4.0f)
-				[
-					SNew(STextBlock)
-						.Text(LOCTEXT("OodleTab_Label_EncodeSpeed", "Encode Speed:"))
-						.ToolTipText(LOCTEXT("OodleTab_Tooltip_EncodeSpeed", "Which of the encode speeds was used for this texture encode, if the encoder supports encode speed."))
-				]
-	
-			+ SVerticalBox::Slot()
-				.AutoHeight()
-				.VAlign(VAlign_Center)
-				.Padding(4.0f)
-				[
-					SAssignNew(OodleRDOEnabledLabel, STextBlock)
-						.Text(LOCTEXT("OodleTab_Label_RDOEnabled", "RDO Lambda:"))
-						.ToolTipText(LOCTEXT("OodleTab_Tooltip_RDOEnabled", "Whether or not the texture was encoded with RDO enabled. If enabled, shows the lambda used to encode. Excludes any global ini specific adjustments (e.g. GlobalLambdaMultiplier)"))
-				]
+	OodleCompressionLevel = FOodleDataCompression::ECompressionLevel::Optimal3;
+	const TCHAR* LevelName;
+	{		 
+		FOodleDataCompression::ECompressionLevelFromValue(ProjectSettings->PackageCompressionLevel_Distribution, OodleCompressionLevel);
+		FOodleDataCompression::ECompressionLevelToString(OodleCompressionLevel, &LevelName);
+	}
 
-			+ SVerticalBox::Slot()
-				.AutoHeight()
-				.VAlign(VAlign_Center)
-				.Padding(4.0f)
-				[
-					SAssignNew(OodleRDOSourceLabel, STextBlock)
-						.Text(LOCTEXT("OodleTab_Label_RDOSource", "RDO Lambda Source:"))
-						.ToolTipText(LOCTEXT("OodleTab_Tooltip_RDOSource", "This is where the build system found the lambda to use, due to defaults and fallbacks. (Lambda) means a direct lambda value (Lossy Compression Amount) means it was converted from that property."))
-				]
-	
-			+ SVerticalBox::Slot()
-				.AutoHeight()
-				.VAlign(VAlign_Center)
-				.Padding(4.0f)
-				[
-					SAssignNew(OodleEffortLabel, STextBlock)
-						.Text(LOCTEXT("OodleTab_Label_Effort", "Effort:"))
-						.ToolTipText(LOCTEXT("OodleTab_ToolTip_Effort", "Which effort value was used when encoding this texture. Pulled from the encode speed options. Effort represents how much CPU time was spent finding better results."))
-				]
+	// Grab the compression block size in the settings.
+	{
+		FString CompBlockSizeString;
+		if (FParse::Value(*ProjectSettings->PackageAdditionalCompressionOptions, TEXT("-compressionblocksize="), CompBlockSizeString) &&
+			FParse::Value(*ProjectSettings->PackageAdditionalCompressionOptions, TEXT("-compressionblocksize="), CompressionBlockSize))
+		{
+			if (CompBlockSizeString.EndsWith(TEXT("MB")))
+			{
+				CompressionBlockSize *= 1024 * 1024;
+			}
+			else if (CompBlockSizeString.EndsWith(TEXT("KB")))
+			{
+				CompressionBlockSize *= 1024;
+			}
+		}
+		else
+		{
+			UE_LOG(LogTextureEditor, Warning, TEXT("No compression block size found in settings - using 256KB"));
+			CompressionBlockSize = 256*1024;
+		}
+	}
 
-			+ SVerticalBox::Slot()
-				.AutoHeight()
-				.VAlign(VAlign_Center)
-				.Padding(4.0f)
-				[
-					SAssignNew(OodleTilingLabel, STextBlock)
-						.Text(LOCTEXT("OodleTab_Label_UniversalTiling", "Universal Tiling:"))
-						.ToolTipText(LOCTEXT("OodleTab_ToolTip_UniversalTiling", "Which universal tiling setting was used when encoding this texture. Specified with encode speed. Universal Tiling is a technique to save on-disc space for platforms that expect tiled textures."))
-				]
-		]
+	TextureViewport = SNew(STextureEditorViewport, SharedThis(this));
 
-		+ SHorizontalBox::Slot()
-		.FillWidth(0.5f)
-		[
-			SNew(SVerticalBox)
+	OodleTabContainer = SNew(SVerticalBox)
 
-			+ SVerticalBox::Slot()
-				.AutoHeight()
-				.VAlign(VAlign_Center)
-				.Padding(4.0f)
-				[
-					SAssignNew(OodleEncoderText, STextBlock)
-				]
+	//
+	// Oodle relevant details container
+	//
+	+ SVerticalBox::Slot()
+		.AutoHeight()
+		.Padding(4.0f)
+		[
+			SNew(SHorizontalBox)
 
-			+ SVerticalBox::Slot()
-				.AutoHeight()
-				.VAlign(VAlign_Center)
-				.Padding(4.0f)
+			//
+			// Details label container
+			//
+			+ SHorizontalBox::Slot()
+				.FillWidth(0.5f)
 				[
-					SAssignNew(OodleEncodeSpeedText, STextBlock)
+					SNew(SVerticalBox)
+					+ SVerticalBox::Slot()
+						.AutoHeight()
+						.VAlign(VAlign_Center)
+						.Padding(4.0f)
+						[
+							SNew(STextBlock)
+								.Text(LOCTEXT("OodleTab_Label_Encoder", "Encoder:"))
+								.ToolTipText(LOCTEXT("OodleTab_Tooltip_Encoder", "Which texture encoder was used to encode the texture."))
+						]
+					+ SVerticalBox::Slot()
+						.AutoHeight()
+						.VAlign(VAlign_Center)
+						.Padding(4.0f)
+						[
+							SNew(STextBlock)
+								.Text(LOCTEXT("OodleTab_Label_EncodeSpeed", "Encode Speed:"))
+								.ToolTipText(LOCTEXT("OodleTab_Tooltip_EncodeSpeed", "Which of the encode speeds was used for this texture encode, if the encoder supports encode speed."))
+						]
+					+ SVerticalBox::Slot()
+						.AutoHeight()
+						.VAlign(VAlign_Center)
+						.Padding(4.0f)
+						[
+							SAssignNew(OodleRDOEnabledLabel, STextBlock)
+								.Text(LOCTEXT("OodleTab_Label_RDOEnabled", "RDO Lambda:"))
+								.ToolTipText(LOCTEXT("OodleTab_Tooltip_RDOEnabled", "Whether or not the texture was encoded with RDO enabled. If enabled, shows the lambda used to encode. Excludes any global ini specific adjustments (e.g. GlobalLambdaMultiplier)"))
+						]
+					+ SVerticalBox::Slot()
+						.AutoHeight()
+						.VAlign(VAlign_Center)
+						.Padding(4.0f)
+						[
+							SAssignNew(OodleRDOSourceLabel, STextBlock)
+								.Text(LOCTEXT("OodleTab_Label_RDOSource", "RDO Lambda Source:"))
+								.ToolTipText(LOCTEXT("OodleTab_Tooltip_RDOSource", "This is where the build system found the lambda to use, due to defaults and fallbacks. (Lambda) means a direct lambda value (Lossy Compression Amount) means it was converted from that property."))
+						]
+					+ SVerticalBox::Slot()
+						.AutoHeight()
+						.VAlign(VAlign_Center)
+						.Padding(4.0f)
+						[
+							SAssignNew(OodleEffortLabel, STextBlock)
+								.Text(LOCTEXT("OodleTab_Label_Effort", "Effort:"))
+								.ToolTipText(LOCTEXT("OodleTab_ToolTip_Effort", "Which effort value was used when encoding this texture. Pulled from the encode speed options. Effort represents how much CPU time was spent finding better results."))
+						]
+
+					+ SVerticalBox::Slot()
+						.AutoHeight()
+						.VAlign(VAlign_Center)
+						.Padding(4.0f)
+						[
+							SAssignNew(OodleTilingLabel, STextBlock)
+								.Text(LOCTEXT("OodleTab_Label_UniversalTiling", "Universal Tiling:"))
+								.ToolTipText(LOCTEXT("OodleTab_ToolTip_UniversalTiling", "Which universal tiling setting was used when encoding this texture. Specified with encode speed. Universal Tiling is a technique to save on-disk space for platforms that expect tiled textures."))
+						]
 				]
 
-			+ SVerticalBox::Slot()
-				.AutoHeight()
-				.VAlign(VAlign_Center)
-				.Padding(4.0f)
+			//
+			// Details controls container
+			//
+			+ SHorizontalBox::Slot()
+				.FillWidth(0.5f)
 				[
-					SAssignNew(OodleRDOText, STextBlock)
+					SNew(SVerticalBox)
+					+ SVerticalBox::Slot()
+						.AutoHeight()
+						.VAlign(VAlign_Center)
+						.Padding(4.0f)
+						[
+							SAssignNew(OodleEncoderText, STextBlock)
+						]
+					+ SVerticalBox::Slot()
+						.AutoHeight()
+						.VAlign(VAlign_Center)
+						.Padding(4.0f)
+						[
+							SAssignNew(OodleEncodeSpeedText, STextBlock)
+						]
+					+ SVerticalBox::Slot()
+						.AutoHeight()
+						.VAlign(VAlign_Center)
+						.Padding(4.0f)
+						[
+							SAssignNew(OodleRDOText, STextBlock)
+						]
+					+ SVerticalBox::Slot()
+						.AutoHeight()
+						.VAlign(VAlign_Center)
+						.Padding(4.0f)
+						[
+							SAssignNew(OodleRDOSourceText, STextBlock)
+						]
+					+ SVerticalBox::Slot()
+						.AutoHeight()
+						.VAlign(VAlign_Center)
+						.Padding(4.0f)
+						[
+							SAssignNew(OodleEffortText, STextBlock)
+						]
+					+ SVerticalBox::Slot()
+						.AutoHeight()
+						.VAlign(VAlign_Center)
+						.Padding(4.0f)
+						[
+							SAssignNew(OodleTilingText, STextBlock)
+						]
 				]
+		] // end oodle details.
 
-			+ SVerticalBox::Slot()
-				.AutoHeight()
-				.VAlign(VAlign_Center)
-				.Padding(4.0f)
+	//
+	// Header for oodle rdo experiments
+	//
+	+ SVerticalBox::Slot()
+		.AutoHeight()
+		.VAlign(VAlign_Center)
+		[
+			SNew(SHeader)
+				.HAlign(EHorizontalAlignment::HAlign_Fill)
 				[
-					SAssignNew(OodleRDOSourceText, STextBlock)
+					SNew(STextBlock)
+					.Text(LOCTEXT("OodleTab_Label_TryHeader", "Try Encodings"))
 				]
+		]
 
-			+ SVerticalBox::Slot()
-				.AutoHeight()
-				.VAlign(VAlign_Center)
-				.Padding(4.0f)
+	//
+	// Container for oodle rdo experiments labels/controls
+	//
+	+ SVerticalBox::Slot()
+		.AutoHeight()
+		.Padding(4.0f)
+		[
+			//
+			// Labels for oodle rdo experiments.
+			//
+			SNew(SHorizontalBox)
+			+ SHorizontalBox::Slot()
+				.FillWidth(0.5f)
 				[
-					SAssignNew(OodleEffortText, STextBlock)
+					SNew(SVerticalBox)
+					+ SVerticalBox::Slot()
+						.AutoHeight()
+						.VAlign(VAlign_Center)
+						.Padding(6)
+						[
+							SNew(STextBlock)
+							.Text(LOCTEXT("OodleTab_Label_OverrideCompression", "Enabled:"))
+							.ToolTipText(LOCTEXT("OodleTab_ToolTip_OverrideCompression", "If checked, allows you to experiment with Oodle RDO compression settings to visualize results."))
+						]
+					+ SVerticalBox::Slot()
+						.AutoHeight()
+						.VAlign(VAlign_Center)
+						.Padding(6)
+						[
+							SNew(STextBlock)
+							.Text(LOCTEXT("OodleTab_Label_OverrideRDO", "RDO Lambda:"))
+							.ToolTipText(LOCTEXT("OodleTab_ToolTip_OverrideRDO", "The RDO lambda to encode with for experimentation. 0 disables RDO entirely. 1 is largest filesize, 100 is smallest."))
+						]
+					+ SVerticalBox::Slot()
+						.AutoHeight()
+						.VAlign(VAlign_Center)
+						.Padding(6)
+						[
+							SNew(STextBlock)
+							.Text(LOCTEXT("OodleTab_Label_OverrideEffort", "Effort:"))
+							.ToolTipText(LOCTEXT("OodleTab_ToolTip_OverrideEffort", "The encoding effort to try. Effort controls how much CPU time spent on finding better results. See the Oodle Texture documentation for detailed information."))
+						]
+					+ SVerticalBox::Slot()
+						.AutoHeight()
+						.VAlign(VAlign_Center)
+						.Padding(6)
+						[
+							SNew(STextBlock)
+							.Text(LOCTEXT("OodleTab_Label_OverrideTiling", "Universal Tiling:"))
+							.ToolTipText(LOCTEXT("OodleTab_ToolTip_OverrideTiling", "The universal tiling to try. See the Oodle Texture documentation for detailed information."))
+						]
 				]
 
-			+ SVerticalBox::Slot()
-				.AutoHeight()
-				.VAlign(VAlign_Center)
-				.Padding(4.0f)
+			//
+			// Controls for oodle rdo experiments
+			//
+			+ SHorizontalBox::Slot()
+				.FillWidth(0.5f)
 				[
-					SAssignNew(OodleTilingText, STextBlock)
-				]
-		]
-	]
+					SNew(SVerticalBox)
+					+ SVerticalBox::Slot()
+						.AutoHeight()
+						.VAlign(VAlign_Center)
+						.Padding(2)
+						[
+							SAssignNew(OodleOverrideCheck, SCheckBox)
+							.OnCheckStateChanged(this, &FTextureEditorToolkit::OnUseEditorOodleSettingsChanged)
+							.IsChecked(this, &FTextureEditorToolkit::UseEditorOodleSettingsChecked)
+						]
+					+ SVerticalBox::Slot()
+						.AutoHeight()
+						.VAlign(VAlign_Center)
+						.Padding(2)
+						[
+							SNew(SNumericEntryBox<int32>)
+							.Value(this, &FTextureEditorToolkit::GetEditorOodleSettingsRDO)
+							.OnValueCommitted(this, &FTextureEditorToolkit::EditorOodleSettingsRDOCommitted)
+							.IsEnabled(this, &FTextureEditorToolkit::EditorOodleSettingsEnabled)
+						]
+					+ SVerticalBox::Slot()
+						.AutoHeight()
+						.VAlign(VAlign_Center)
+						.Padding(2)
+						[
+							SNew(SEnumComboBox, StaticEnum< ETextureEncodeEffort >())
+							.CurrentValue(this, &FTextureEditorToolkit::GetEditorOodleSettingsEffort)
+							.OnEnumSelectionChanged(this, &FTextureEditorToolkit::EditorOodleSettingsEffortChanged)
+							.IsEnabled(this, &FTextureEditorToolkit::EditorOodleSettingsEnabled)
+						]
+					+ SVerticalBox::Slot()
+						.AutoHeight()
+						.VAlign(VAlign_Center)
+						.Padding(2)
+						[
+							SNew(SEnumComboBox, StaticEnum< ETextureUniversalTiling >())
+							.CurrentValue(this, &FTextureEditorToolkit::GetEditorOodleSettingsTiling)
+							.OnEnumSelectionChanged(this, &FTextureEditorToolkit::EditorOodleSettingsTilingChanged)
+							.IsEnabled(this, &FTextureEditorToolkit::EditorOodleSettingsEnabled)
+						]
+				] // end controls
+		] // end oodle rdo experiment labels/controls
+
+	//
+	// Header for the on disk estimates
+	//
 	+ SVerticalBox::Slot()
 		.AutoHeight()
 		.VAlign(VAlign_Center)
 		[
 			SNew(SHeader)
-				.HAlign(EHorizontalAlignment::HAlign_Fill)
+			.HAlign(EHorizontalAlignment::HAlign_Fill)
 			[
 				SNew(STextBlock)
-				.Text(LOCTEXT("OodleTab_Label_TryHeader", "Try Encodings"))
+				.Text(LOCTEXT("OodleTab_Label_EstimatesHeader", "On-disk Sizes"))
+				.ToolTipText(LOCTEXT("OodleTab_ToolTip_EstimatesHeader", "RDO encoding only helps on-disk texture sizes when package compression is enabled. It does not affect runtime memory usage."))
 			]
 		]
-
+	//
+	// Container for the on disk estimate labels/controls.
+	//
 	+ SVerticalBox::Slot()
 		.AutoHeight()
 		.Padding(4.0f)
 		[
+			//
+			// Labels for the on disk estimates
+			//
 			SNew(SHorizontalBox)
 			+ SHorizontalBox::Slot()
 				.FillWidth(0.5f)
 				[
-					SNew(SVerticalBox)						
-
-						+ SVerticalBox::Slot()
-							.AutoHeight()
-							.VAlign(VAlign_Center)
-							.Padding(6)
-							[
-								SNew(STextBlock)
-								.Text(LOCTEXT("OodleTab_Label_OverrideCompression", "Enabled:"))
-							.ToolTipText(LOCTEXT("OodleTab_ToolTip_OverrideCompression", "If checked, allows you to experiment with Oodle RDO compression settings to visualize results."))
-
-							]
-						+ SVerticalBox::Slot()
-							.AutoHeight()
-							.VAlign(VAlign_Center)
-							.Padding(6)
-							[
-								SNew(STextBlock)
-								.Text(LOCTEXT("OodleTab_Label_OverrideRDO", "RDO Lambda:"))
-							.ToolTipText(LOCTEXT("OodleTab_ToolTip_OverrideRDO", "The RDO lambda to encode with for experimentation. 0 disables RDO entirely. 1 is largest filesize, 100 is smallest."))
-
-							]
-						+ SVerticalBox::Slot()
-							.AutoHeight()
-							.VAlign(VAlign_Center)
-							.Padding(6)
-							[
-								SNew(STextBlock)
-								.Text(LOCTEXT("OodleTab_Label_OverrideEffort", "Effort:"))
-							.ToolTipText(LOCTEXT("OodleTab_ToolTip_OverrideEffort", "The encoding effort to try. Effort controls how much CPU time spent on finding better results. See the Oodle Texture documentation for detailed information."))
-
-							]
-						+ SVerticalBox::Slot()
-							.AutoHeight()
-							.VAlign(VAlign_Center)
-							.Padding(6)
-							[
-								SNew(STextBlock)
-								.Text(LOCTEXT("OodleTab_Label_OverrideTiling", "Universal Tiling:"))
-							.ToolTipText(LOCTEXT("OodleTab_ToolTip_OverrideTiling", "The universal tiling to try. See the Oodle Texture documentation for detailed information."))
-
-							]
+					SNew(SVerticalBox)
+					+ SVerticalBox::Slot()
+						.AutoHeight()
+						.VAlign(VAlign_Center)
+						.Padding(6)
+						[
+							SNew(STextBlock)
+							.Text(LOCTEXT("OodleTab_Label_EstimatesEnabled", "Enabled:"))
+							.ToolTipText(LOCTEXT("OodleTab_ToolTip_EstimatesEnabled", "If checked, texture data will be compressed in the same manner as project packaging in order to estimate the benefits of RDO encoding of the texture."))
+						]
+					+ SVerticalBox::Slot()
+						.AutoHeight()
+						.VAlign(VAlign_Center)
+						.Padding(6)
+						[
+							SNew(STextBlock)
+							.Text(LOCTEXT("OodleTab_Label_EncoderSettings", "Packaging Configuration:"))
+							.ToolTipText(LOCTEXT("OodleTab_ToolTip_EncoderSettings", "Which packaging configuration to pull from for determining which Oodle encoder and compression level to use."))
+						]
+					+ SVerticalBox::Slot()
+						.AutoHeight()
+						.VAlign(VAlign_Center)
+						.Padding(6)
+						[
+							SNew(STextBlock)
+							.Text(LOCTEXT("OodleTab_Label_EstimateEncoder", "Oodle Encoder:"))
+							.ToolTipText(LOCTEXT("OodleTab_ToolTip_EstimateEncoder", "The oodle encoder to use for estimating. Pulled from the packaging configuration specified above."))
+						]
+					+ SVerticalBox::Slot()
+						.AutoHeight()
+						.VAlign(VAlign_Center)
+						.Padding(6)
+						[
+							SNew(STextBlock)
+							.Text(LOCTEXT("OodleTab_Label_EstimateLevel", "Oodle Compression Level:"))
+							.ToolTipText(LOCTEXT("OodleTab_ToolTip_EstimateLevel", "The compression level. Pulled from the packaging configuration specified above."))
+						]
+					+ SVerticalBox::Slot()
+						.AutoHeight()
+						.VAlign(VAlign_Center)
+						.Padding(6)
+						[
+							SNew(STextBlock)
+							.Text(LOCTEXT("OodleTab_Label_BlockSize", "Compression Block Size:"))
+							.ToolTipText(LOCTEXT("OodleTab_ToolTip_BlockSize", "The size of chunks used when compressing. Pulled from the packaging configuration 'Package Compression Commandline Options'."))
+						]
+					+ SVerticalBox::Slot()
+						.AutoHeight()
+						.VAlign(VAlign_Center)
+						.Padding(6)
+						[
+							SNew(STextBlock)
+							.Text(LOCTEXT("OodleTab_Label_EstimateRaw", "Uncompressed size:"))
+							.ToolTipText(LOCTEXT("OodleTab_ToolTip_EstimateRaw", "The size of the mip or virtual texture data for the texture."))
+						]
+					+ SVerticalBox::Slot()
+						.AutoHeight()
+						.VAlign(VAlign_Center)
+						.Padding(6)
+						[
+							SNew(STextBlock)
+							.Text(LOCTEXT("OodleTab_Label_EstimateCompressed", "Compressed size (estimate):"))
+							.ToolTipText(LOCTEXT("OodleTab_ToolTip_EstimateCompressed", "The size of the compressed mip or virtual texture data for the texture."))
+						]
 				]
 
+			//
+			// Controls for the on disk estimates
+			//
 			+ SHorizontalBox::Slot()
 				.FillWidth(0.5f)
 				[
 					SNew(SVerticalBox)
-
-
-						+ SVerticalBox::Slot()
-							.AutoHeight()
-							.VAlign(VAlign_Center)
-							.Padding(2)
-							[
-								SAssignNew(OodleOverrideCheck, SCheckBox)
-								.OnCheckStateChanged(this, &FTextureEditorToolkit::OnUseEditorOodleSettingsChanged)
-								.IsChecked(this, &FTextureEditorToolkit::UseEditorOodleSettingsChecked)
-							]
-
-						+ SVerticalBox::Slot()
-							.AutoHeight()
-							.VAlign(VAlign_Center)
-							.Padding(2)
-							[
-								SNew(SNumericEntryBox<int32>)
-								.Value(this, &FTextureEditorToolkit::GetEditorOodleSettingsRDO)
-								.OnValueCommitted(this, &FTextureEditorToolkit::EditorOodleSettingsRDOCommitted)
-								.IsEnabled(this, &FTextureEditorToolkit::EditorOodleSettingsEnabled)
-							]
-
-						+ SVerticalBox::Slot()
-							.AutoHeight()
-							.VAlign(VAlign_Center)
-							.Padding(2)
-							[
-								SNew(SEnumComboBox, StaticEnum< ETextureEncodeEffort >())
-								.CurrentValue(this, &FTextureEditorToolkit::GetEditorOodleSettingsEffort)
-								.OnEnumSelectionChanged(this, &FTextureEditorToolkit::EditorOodleSettingsEffortChanged)
-								.IsEnabled(this, &FTextureEditorToolkit::EditorOodleSettingsEnabled)
-							]
-
-						+ SVerticalBox::Slot()
-							.AutoHeight()
-							.VAlign(VAlign_Center)
-							.Padding(2)
-							[
-								SNew(SEnumComboBox, StaticEnum< ETextureUniversalTiling >())
-								.CurrentValue(this, &FTextureEditorToolkit::GetEditorOodleSettingsTiling)
-								.OnEnumSelectionChanged(this, &FTextureEditorToolkit::EditorOodleSettingsTilingChanged)
-								.IsEnabled(this, &FTextureEditorToolkit::EditorOodleSettingsEnabled)
-							]
+					+ SVerticalBox::Slot()
+						.AutoHeight()
+						.VAlign(VAlign_Center)
+						.Padding(2)
+						[
+							SAssignNew(OodleEstimateCheck, SCheckBox)
+							.OnCheckStateChanged(this, &FTextureEditorToolkit::OnEstimateCompressionChanged)
+							.IsChecked(this, &FTextureEditorToolkit::EstimateCompressionChecked)
+						]
+					+ SVerticalBox::Slot()
+						.AutoHeight()
+						.VAlign(VAlign_Center)
+						.Padding(2)
+						[
+							SNew(STextComboBox)
+							.OptionsSource(&PackagingSettingsNames)
+							.OnSelectionChanged(this, &FTextureEditorToolkit::PackagingSettingsChanged)
+							.IsEnabled(this, &FTextureEditorToolkit::EstimateCompressionEnabled)
+							.InitiallySelectedItem(InitialPackagingSetting)
+						]
+					+ SVerticalBox::Slot()
+						.AutoHeight()
+						.VAlign(VAlign_Center)
+						.Padding(8)
+						[
+							SAssignNew(OodleEncoderUsed, STextBlock)
+							.Text(FText::AsCultureInvariant(CompressorName))
+							.IsEnabled(this, &FTextureEditorToolkit::EstimateCompressionEnabled)
+						]
+					+ SVerticalBox::Slot()
+						.AutoHeight()
+						.VAlign(VAlign_Center)
+						.Padding(6)
+						[
+							SAssignNew(OodleLevelUsed, STextBlock)
+							.Text(FText::FromString(FString::Printf(TEXT("%s (%d)"), LevelName, (int8)OodleCompressionLevel)))
+							.IsEnabled(this, &FTextureEditorToolkit::EstimateCompressionEnabled)
+						]
+					+ SVerticalBox::Slot()
+						.AutoHeight()
+						.VAlign(VAlign_Center)
+						.Padding(6)
+						[
+							SAssignNew(OodleCompressionBlockUsed, STextBlock)
+							.Text(FText::AsMemory(CompressionBlockSize))
+							.IsEnabled(this, &FTextureEditorToolkit::EstimateCompressionEnabled)
+						]
+					+ SVerticalBox::Slot()
+						.AutoHeight()
+						.VAlign(VAlign_Center)
+						.Padding(6)
+						[
+							SAssignNew(OodleEstimateRaw, STextBlock)
+							.IsEnabled(this, &FTextureEditorToolkit::EstimateCompressionEnabled)
+						]
+					+ SVerticalBox::Slot()
+						.AutoHeight()
+						.VAlign(VAlign_Center)
+						.Padding(6)
+						[
+							SAssignNew(OodleEstimateCompressed, STextBlock)
+							.IsEnabled(this, &FTextureEditorToolkit::EstimateCompressionEnabled)
+						]
 				]
-		];
-
-
-
+		]; // end on disk estimates controls
 
 
 	TextureProperties = SNew(SVerticalBox)
@@ -2115,11 +2404,13 @@ int32 FTextureEditorToolkit::GetEditorOodleSettingsEffort() const
 
 void FTextureEditorToolkit::EditorOodleSettingsEffortChanged(int32 NewValue, ESelectInfo::Type SelectionType)
 {
+	bool bChanged = CustomEncoding->OodleEncodeEffort != NewValue;
+
 	CustomEncoding->OodleEncodeEffort = NewValue;
 
-	if (CustomEncoding->bUseCustomEncode)
+	if (CustomEncoding->bUseCustomEncode || bChanged)
 	{
-		PostTextureRecode(Texture);
+		PostTextureRecode();
 	}
 }
 
@@ -2130,11 +2421,12 @@ int32 FTextureEditorToolkit::GetEditorOodleSettingsTiling() const
 
 void FTextureEditorToolkit::EditorOodleSettingsTilingChanged(int32 NewValue, ESelectInfo::Type SelectionType)
 {
+	bool bChanged = CustomEncoding->OodleUniversalTiling != NewValue;
 	CustomEncoding->OodleUniversalTiling = NewValue;
 
-	if (CustomEncoding->bUseCustomEncode)
+	if (CustomEncoding->bUseCustomEncode && bChanged)
 	{
-		PostTextureRecode(Texture);
+		PostTextureRecode();
 	}
 }
 
@@ -2154,11 +2446,13 @@ void FTextureEditorToolkit::EditorOodleSettingsRDOCommitted(int32 NewValue, ETex
 		NewValue = 0;
 	}
 
+	bool bChanged = CustomEncoding->OodleRDOLambda != (int8)NewValue;
+
 	CustomEncoding->OodleRDOLambda = (int8)NewValue;
 
-	if (CustomEncoding->bUseCustomEncode)
+	if (CustomEncoding->bUseCustomEncode && bChanged)
 	{
-		PostTextureRecode(Texture);
+		PostTextureRecode();
 	}
 }
 
@@ -2179,7 +2473,7 @@ void FTextureEditorToolkit::OnUseEditorOodleSettingsChanged(ECheckBoxState NewSt
 	// that they need to update, so we do this by faking a compression method property change.
 	CustomEncoding->bUseCustomEncode = NewState == ECheckBoxState::Checked ? true : false;
 
-	PostTextureRecode(Texture);
+	PostTextureRecode();
 }
 
 TSharedRef<SWidget> FTextureEditorToolkit::MakeChannelControlWidget()
@@ -2513,4 +2807,47 @@ TSharedRef<SWidget> FTextureEditorToolkit::MakeZoomControlWidget()
 	return ZoomControl;
 }
 
+void FTextureEditorToolkit::OnEstimateCompressionChanged(ECheckBoxState NewState)
+{
+	OodleCompressedPreviewDDCKey.Set<FString>(FString());
+	bEstimateCompressionEnabled = NewState == ECheckBoxState::Checked;
+}
+ECheckBoxState FTextureEditorToolkit::EstimateCompressionChecked() const
+{
+	return bEstimateCompressionEnabled ? ECheckBoxState::Checked : ECheckBoxState::Unchecked;
+}
+
+void FTextureEditorToolkit::PackagingSettingsChanged(TSharedPtr<FString> Selection, ESelectInfo::Type SelectInfo)
+{
+	if (Selection.IsValid())
+	{
+		UProjectPackagingSettings const* ProjectSettings = GetDefault<UProjectPackagingSettings>();
+		int8 CompressionLevelFromSettings = (int8)FOodleDataCompression::ECompressionLevel::Optimal3;
+		if (*Selection == TEXT("DebugDevelopment"))
+		{
+			CompressionLevelFromSettings = ProjectSettings->PackageCompressionLevel_DebugDevelopment;
+		}
+		else if (*Selection == TEXT("TestShipping"))
+		{
+			CompressionLevelFromSettings = ProjectSettings->PackageCompressionLevel_TestShipping;
+		}
+		else if (*Selection == TEXT("Distribution"))
+		{
+			CompressionLevelFromSettings = ProjectSettings->PackageCompressionLevel_Distribution;
+		}
+
+		FOodleDataCompression::ECompressionLevel OldLevel = OodleCompressionLevel;
+		FOodleDataCompression::ECompressionLevelFromValue(CompressionLevelFromSettings, OodleCompressionLevel);
+
+		const TCHAR* LevelName;
+		FOodleDataCompression::ECompressionLevelToString(OodleCompressionLevel, &LevelName);
+		OodleLevelUsed->SetText(FText::FromString(FString::Printf(TEXT("%s (%d)"), LevelName, CompressionLevelFromSettings)));
+
+		if (OldLevel != OodleCompressionLevel)
+		{
+			OodleCompressedPreviewDDCKey.Set<FString>(FString());
+		}
+	}
+}
+
 #undef LOCTEXT_NAMESPACE
diff --git a/Engine/Source/Editor/TextureEditor/Private/TextureEditorToolkit.h b/Engine/Source/Editor/TextureEditor/Private/TextureEditorToolkit.h
--- a/Engine/Source/Editor/TextureEditor/Private/TextureEditorToolkit.h
+++ b/Engine/Source/Editor/TextureEditor/Private/TextureEditorToolkit.h
@@ -17,6 +17,8 @@ class STextureEditorViewport;
 class UFactory;
 class UTexture;
 
+namespace FOodleDataCompression { enum class ECompressor : uint8; enum class ECompressionLevel : int8; }
+
 enum class ETextureChannelButton : uint8
 {
 	Red,
@@ -420,6 +422,43 @@ private:
 	int32 GetEditorOodleSettingsTiling() const;
 	void EditorOodleSettingsTilingChanged(int32 NewValue, ESelectInfo::Type SelectionType);
 
+	// This is the key we have compressed for oodle preview, so we can catch changes.
+	TVariant<FString, UE::DerivedData::FCacheKeyProxy> OodleCompressedPreviewDDCKey;
+
+	
+	//
+	// -- Start on disc estimation controls/fns
+	//
+	bool bEstimateCompressionEnabled = false;
+	bool EstimateCompressionEnabled() const { return bEstimateCompressionEnabled; }
+
+	// Parsed from packaging settings.
+	FOodleDataCompression::ECompressionLevel OodleCompressionLevel;
+	FOodleDataCompression::ECompressor OodleCompressor;
+	uint32 CompressionBlockSize;
+
+	// Valid if we've kicked off an estimation task and are waiting for the result.
+	TFuture<TTuple<uint64, uint64>> OutstandingEstimation;
+
+	// Controls
+	TSharedPtr<class SCheckBox> OodleEstimateCheck;
+	void OnEstimateCompressionChanged(ECheckBoxState NewState);
+	ECheckBoxState EstimateCompressionChecked() const;
+
+	TSharedPtr<STextBlock> OodleEncoderUsed;
+	TSharedPtr<STextBlock> OodleLevelUsed;
+	TSharedPtr<STextBlock> OodleCompressionBlockUsed;
+	TSharedPtr<STextBlock> OodleEstimateRaw;
+	TSharedPtr<STextBlock> OodleEstimateCompressed;
+
+	// For the packaging settings combo box 
+	TArray< TSharedPtr< FString > > PackagingSettingsNames;
+	void PackagingSettingsChanged(TSharedPtr<FString> Selection, ESelectInfo::Type SelectInfo);
+	//
+	// -- End on disc estimation controls/fns
+	//
+
+	void PostTextureRecode();
 public:
 	// These are on a separate object so then engine can share the type without needing
 	// our module.
diff --git a/Engine/Source/Editor/TextureEditor/TextureEditor.Build.cs b/Engine/Source/Editor/TextureEditor/TextureEditor.Build.cs
--- a/Engine/Source/Editor/TextureEditor/TextureEditor.Build.cs
+++ b/Engine/Source/Editor/TextureEditor/TextureEditor.Build.cs
@@ -47,7 +47,9 @@ public class TextureEditor : ModuleRules
                 "PropertyEditor",
 				"MediaUtils",
 				"EditorWidgets",
-				"MediaAssets"
+				"MediaAssets",
+				"DerivedDataCache",
+				"DeveloperToolSettings"
 			}
 		);
 	}
diff --git a/Engine/Source/Runtime/Engine/Classes/Engine/Texture.h b/Engine/Source/Runtime/Engine/Classes/Engine/Texture.h
--- a/Engine/Source/Runtime/Engine/Classes/Engine/Texture.h
+++ b/Engine/Source/Runtime/Engine/Classes/Engine/Texture.h
@@ -38,6 +38,9 @@
 
 #include "Texture.generated.h"
 
+namespace FOodleDataCompression {enum class ECompressor : uint8; enum class ECompressionLevel : int8; }
+
+
 class ITargetPlatform;
 class UAssetUserData;
 struct FPropertyChangedEvent;
@@ -855,6 +858,11 @@ public:
 	bool TryCancelCache();
 	void CancelCache();
 	ENGINE_API bool TryInlineMipData(int32 FirstMipToLoad = 0, FStringView DebugContext=FStringView());
+	ENGINE_API TFuture<TTuple<uint64, uint64>> LaunchEstimateOnDiskSizeTask(
+		FOodleDataCompression::ECompressor InOodleCompressor,
+		FOodleDataCompression::ECompressionLevel InOodleCompressionLevel,
+		uint32 InCompressionBlockSize,
+		FStringView InDebugContext);
 	bool AreDerivedMipsAvailable(FStringView Context) const;
 	bool AreDerivedVTChunksAvailable(FStringView Context) const;
 	UE_DEPRECATED(5.00, "Use AreDerivedMipsAvailable with the context instead.")
diff --git a/Engine/Source/Runtime/Engine/Private/TextureDerivedData.cpp b/Engine/Source/Runtime/Engine/Private/TextureDerivedData.cpp
--- a/Engine/Source/Runtime/Engine/Private/TextureDerivedData.cpp
+++ b/Engine/Source/Runtime/Engine/Private/TextureDerivedData.cpp
@@ -50,6 +50,8 @@
 #include "TextureCompiler.h"
 #include "TextureEncodingSettings.h"
 #include "ColorSpace.h"
+#include "Compression/OodleDataCompressionUtil.h"
+
 
 /*------------------------------------------------------------------------------
 	Versioning for texture derived data.
@@ -1365,6 +1367,313 @@ static void CheckMipSize(FTexture2DMipMap& Mip, EPixelFormat PixelFormat, int32
 	#endif
 }
 
+//
+// Retrieve all built texture data in to the associated arrays, and don't return unless there's an error
+// or we have the data.
+//
+static bool FetchAllTextureDataSynchronous(FTexturePlatformData* PlatformData, FStringView DebugContext, TArray<TArray<uint8>>& OutMipData, TArray<TArray<uint8>>& OutVTChunkData)
+{
+	OutMipData.Empty();
+	OutVTChunkData.Empty();
+
+	using namespace UE::DerivedData;
+	FDerivedDataCacheInterface& DDC = GetDerivedDataCacheRef();
+
+	OutMipData.AddDefaulted(PlatformData->Mips.Num());
+
+	// This only handles non-vt mips that are paged to derived data.
+	// (some mips are inline). Doesn't handle excluded mips due to 
+	// platform settings.
+	FAsyncMipHandles MipHandles;
+	if (!BeginLoadDerivedMips(*PlatformData, 0, DebugContext, MipHandles,
+			[&OutMipData](int32 MipIndex, FSharedBuffer MipBuffer)
+				{
+					OutMipData[MipIndex].Append((uint8*)MipBuffer.GetData(), MipBuffer.GetSize());
+				}
+			)
+		)
+	{
+		return false;
+	}
+
+	// DDC1 fetches are async, so we need to wait on the handles,
+	// and if it's not paged to derived data we need to copy the
+	// data from bulk
+	for (int32 MipIndex = 0; MipIndex < PlatformData->Mips.Num(); ++MipIndex)
+	{
+		FTexture2DMipMap& Mip = PlatformData->Mips[MipIndex];
+		if (Mip.IsPagedToDerivedData() == false)
+		{
+			OutMipData[MipIndex].Append((uint8*)Mip.BulkData.LockReadOnly(), Mip.BulkData.GetBulkDataSize());
+			Mip.BulkData.Unlock();
+			continue;
+		}
+
+		// Here we either got the data synchronously with the DDC2 path in BeginLoadDerivedMips,
+		// or we didn't if we are DDC1. BeginLoadDerivedMips only allocates the async handles
+		// in the ddc1 path, so we use that to tell.
+		if (MipHandles.Num() == 0)
+		{
+			// skip, ddc1. This is the same for all mips but we still need to check for paged
+			// data above.
+			continue;
+		}
+
+		uint32 AsyncHandle = MipHandles[MipIndex];
+		DDC.WaitAsynchronousCompletion(AsyncHandle);
+		bool bLoadedFromDDC = DDC.GetAsynchronousResults(AsyncHandle, OutMipData[MipIndex]);
+		if (bLoadedFromDDC == false)
+		{
+			return false;
+		}
+	}
+	
+	if (PlatformData->VTData)
+	{
+		OutVTChunkData.AddDefaulted(PlatformData->VTData->Chunks.Num());
+
+		FAsyncVTChunkHandles AsyncVTHandles;
+		BeginLoadDerivedVTChunks(PlatformData->VTData->Chunks, DebugContext, AsyncVTHandles);
+
+		for (int32 ChunkIndex = 0; ChunkIndex < PlatformData->VTData->Chunks.Num(); ++ChunkIndex)
+		{
+			FVirtualTextureDataChunk& Chunk = PlatformData->VTData->Chunks[ChunkIndex];
+			if (Chunk.DerivedDataKey.IsEmpty())
+			{
+				// The data is resident and we can just copy it.
+				OutVTChunkData[ChunkIndex].Append((uint8*)Chunk.BulkData.LockReadOnly(), Chunk.BulkData.GetBulkDataSize());
+				Chunk.BulkData.Unlock();
+			}
+			else
+			{
+				// The data was fetched and we need to wait on the result.
+				uint32 AsyncHandle = AsyncVTHandles[ChunkIndex];
+				DDC.WaitAsynchronousCompletion(AsyncHandle);
+				bool bLoadedFromDDC = DDC.GetAsynchronousResults(AsyncHandle, OutVTChunkData[ChunkIndex]);
+
+				if (bLoadedFromDDC == false)
+				{
+					return false;
+				}
+			}
+		} // end each vt chunk
+	} // end if virtual texture
+
+	return true;
+}
+
+//
+// Chunk the input data in to blocks of the compression block size, then
+// run Oodle on the separate chunks in order to get an estimate of how
+// much space on disk the texture will take during deployment. This
+// exists so the editor can show the benefits of increasing RDO levels 
+// on a texture.
+//
+// This is not exact! Due to the nature of iostore, we can't know exactly
+// whether our data will be chunked on the boundaries we've chosen. However
+// it is illustrative.
+//
+static void EstimateOnDiskCompressionForTextureData(
+	TArray<TArray<uint8>> InMipData,
+	TArray<TArray<uint8>> InVTChunkData,
+	FOodleDataCompression::ECompressor InOodleCompressor,
+	FOodleDataCompression::ECompressionLevel InOodleCompressionLevel,
+	uint32 InCompressionBlockSize,
+	uint64& OutUncompressedByteCount,
+	uint64& OutCompressedByteCount
+)
+{
+	//
+	// This is written such that you can have both classic mip data and
+	// virtual texture data, however actual unreal textures don't have
+	// both.
+	//
+	uint64 UncompressedByteCount = 0;
+	for (TArray<uint8>& Mip : InMipData)
+	{
+		UncompressedByteCount += Mip.Num();
+	}
+	for (TArray<uint8>& VTChunk : InVTChunkData)
+	{
+		UncompressedByteCount += VTChunk.Num();
+	}
+
+	OutUncompressedByteCount = UncompressedByteCount;
+
+	if (UncompressedByteCount == 0)
+	{
+		OutCompressedByteCount = 0;
+		return;
+	}
+
+	int32 MipIndex = 0;
+	int32 VTChunkIndex = 0;
+	int32 CurrentOffsetInContainer = 0;
+	uint64 CompressedByteCount = 0;
+
+	// Array for compressed data so we don't have to realloc.
+	TArray<uint8> Compressed;
+	Compressed.Reserve(InCompressionBlockSize + 1024);
+
+	// When we cross our input array boundaries, we accumulate in to here.
+	TArray<uint8> ContinuousMemory;
+	for (;;)
+	{
+		TArray<uint8>& CurrentContainer = MipIndex < InMipData.Num() ? InMipData[MipIndex] : InVTChunkData[VTChunkIndex];
+
+		uint32 NeedBytes = InCompressionBlockSize - ContinuousMemory.Num();
+		uint32 CopyBytes = CurrentContainer.Num() - CurrentOffsetInContainer;
+		if (CopyBytes > NeedBytes)
+		{
+			CopyBytes = NeedBytes;
+		}
+
+		// Can we compressed without an intervening copy?
+		if (NeedBytes == InCompressionBlockSize && // We don't have a partial block copied
+			CopyBytes == InCompressionBlockSize) // we can fit in this chunk
+		{
+			// Direct.
+			Compressed.Empty();
+			FOodleCompressedArray::CompressData(
+				Compressed,
+				CurrentContainer.GetData() + CurrentOffsetInContainer,
+				InCompressionBlockSize,
+				InOodleCompressor,
+				InOodleCompressionLevel);
+
+			CompressedByteCount += Compressed.Num();
+		}
+		else
+		{
+			// Need to accumulate in to our temp buffer.
+
+			if (ContinuousMemory.Num() == 0)
+			{
+				ContinuousMemory.Reserve(InCompressionBlockSize);
+			}
+
+			ContinuousMemory.Append(CurrentContainer.GetData() + CurrentOffsetInContainer, CopyBytes);
+
+			if (ContinuousMemory.Num() == InCompressionBlockSize)
+			{
+				// Filled a block - kick.
+				Compressed.Empty();
+				FOodleCompressedArray::CompressData(
+					Compressed,
+					ContinuousMemory.GetData(),
+					InCompressionBlockSize,
+					InOodleCompressor,
+					InOodleCompressionLevel);
+
+				CompressedByteCount += Compressed.Num();
+				ContinuousMemory.Empty();
+			}
+		}
+
+		// Advance read cursor.
+		CurrentOffsetInContainer += CopyBytes;
+		if (CurrentOffsetInContainer >= CurrentContainer.Num())
+		{
+			CurrentOffsetInContainer = 0;
+
+			if (MipIndex < InMipData.Num())
+			{
+				MipIndex++;
+			}
+			else if (VTChunkIndex < InVTChunkData.Num())
+			{
+				VTChunkIndex++;
+			}
+
+			if (MipIndex >= InMipData.Num() && VTChunkIndex >= InVTChunkData.Num())
+			{
+				// No more source data.
+				break;
+			}
+		}
+	}
+
+	if (ContinuousMemory.Num())
+	{
+		// If we ran out of source data before we completely filled, kick here.
+		Compressed.Empty();
+		FOodleCompressedArray::CompressData(
+			Compressed,
+			ContinuousMemory.GetData(),
+			ContinuousMemory.Num(),
+			InOodleCompressor,
+			InOodleCompressionLevel);
+
+		CompressedByteCount += Compressed.Num();
+	}
+
+	OutCompressedByteCount = CompressedByteCount;
+}
+
+//
+// Grabs the texture data and then kicks off a task to block compress it
+// in order to try and mimic how iostore does on disk compression.
+//
+// Returns the future result of the compression, with the compressed byte count
+// in the first of the pair and the total in the second.
+//
+TFuture<TTuple<uint64, uint64>> FTexturePlatformData::LaunchEstimateOnDiskSizeTask(
+	FOodleDataCompression::ECompressor InOodleCompressor,
+	FOodleDataCompression::ECompressionLevel InOodleCompressionLevel,
+	uint32 InCompressionBlockSize,
+	FStringView InDebugContext
+	)
+{
+	TArray<TArray<uint8>> MipData;
+	TArray<TArray<uint8>> VTChunkData;
+	if (FetchAllTextureDataSynchronous(this, InDebugContext, MipData, VTChunkData) == false)
+	{
+		return TFuture<TTuple<uint64, uint64>>();
+	}
+	
+	struct FAsyncEstimateState
+	{
+		TPromise<TPair<uint64, uint64>> Promise;
+		TArray<TArray<uint8>> MipData;
+		TArray<TArray<uint8>> VTChunkData;
+		FOodleDataCompression::ECompressor OodleCompressor;
+		FOodleDataCompression::ECompressionLevel OodleCompressionLevel;
+		uint32 CompressionBlockSize;
+	};
+	
+	FAsyncEstimateState* State = new FAsyncEstimateState();
+	State->MipData = MoveTemp(MipData);
+	State->VTChunkData = MoveTemp(VTChunkData);
+	State->OodleCompressor = InOodleCompressor;
+	State->OodleCompressionLevel = InOodleCompressionLevel;
+	State->CompressionBlockSize = InCompressionBlockSize;
+
+	// Grab the future before we kick the task so there's no race.
+	// (unlikely since compression is so long...)
+	TFuture<TTuple<uint64, uint64>> ResultFuture = State->Promise.GetFuture();
+
+	// Kick off a task with no dependencies that does the compression
+	// and posts the result to the future.
+	FGraphEventRef Task = FFunctionGraphTask::CreateAndDispatchWhenReady([State]()
+	{
+		uint64 CompressedByteCount=0, UncompressedByteCount=0;
+
+		EstimateOnDiskCompressionForTextureData(
+			State->MipData, 
+			State->VTChunkData,
+			State->OodleCompressor,
+			State->OodleCompressionLevel,
+			State->CompressionBlockSize,
+			UncompressedByteCount,
+			CompressedByteCount);
+
+		State->Promise.SetValue(TTuple<uint64, uint64>(CompressedByteCount, UncompressedByteCount));
+		delete State;
+	}, TStatId(), nullptr, ENamedThreads::AnyBackgroundThreadNormalTask);
+
+	return ResultFuture;
+}
+
 bool FTexturePlatformData::TryInlineMipData(int32 FirstMipToLoad, FStringView DebugContext)
 {
 	TRACE_CPUPROFILER_EVENT_SCOPE(FTexturePlatformData::TryInlineMipData);
diff --git a/Engine/Source/Runtime/Engine/Public/DerivedDataCacheKeyProxy.h b/Engine/Source/Runtime/Engine/Public/DerivedDataCacheKeyProxy.h
--- a/Engine/Source/Runtime/Engine/Public/DerivedDataCacheKeyProxy.h
+++ b/Engine/Source/Runtime/Engine/Public/DerivedDataCacheKeyProxy.h
@@ -23,7 +23,7 @@ namespace Private
 struct FCacheKeyProxy : private TAlignedBytes<sizeof(UE::DerivedData::Private::FCacheKeyDummy), alignof(UE::DerivedData::Private::FCacheKeyDummy)>
 {
 	FCacheKeyProxy(const FCacheKey& InKey);
-	~FCacheKeyProxy();
+	ENGINE_API ~FCacheKeyProxy();
 	FCacheKey* AsCacheKey() { return (FCacheKey*)this;  }
 	const FCacheKey* AsCacheKey() const { return (const FCacheKey*)this; }
 };
