Fix for brand new texture assets having a null TexturePlatformData. New textures seem to not have CachePlatformData called on them, so the empty TexturePlatformData isn't created.

diff --git a/Engine/Source/Editor/TextureEditor/Private/TextureEditorToolkit.cpp b/Engine/Source/Editor/TextureEditor/Private/TextureEditorToolkit.cpp
--- a/Engine/Source/Editor/TextureEditor/Private/TextureEditorToolkit.cpp
+++ b/Engine/Source/Editor/TextureEditor/Private/TextureEditorToolkit.cpp
@@ -511,7 +511,7 @@ void FTextureEditorToolkit::PopulateQuickInfo( )
 	}
 
 	FTexturePlatformData** PlatformDataPtr = Texture->GetRunningPlatformData();
-	if (PlatformDataPtr)
+	if (PlatformDataPtr && PlatformDataPtr[0]) // Can be null if we haven't had a chance to call CachePlatformData on the texture (brand new)
 	{
 		FTexturePlatformData::FTextureEncodeResultMetadata const& ResultMetadata = PlatformDataPtr[0]->ResultMetadata;
 		if (ResultMetadata.bIsValid == false)
