TextureEditorToolket quiet No compression block size message doesn't need to be a warning
having no setting is standard behavior
also change default size to 64k to match other defaults

diff --git a/Engine/Source/Editor/TextureEditor/Private/TextureEditorToolkit.cpp b/Engine/Source/Editor/TextureEditor/Private/TextureEditorToolkit.cpp
--- a/Engine/Source/Editor/TextureEditor/Private/TextureEditorToolkit.cpp
+++ b/Engine/Source/Editor/TextureEditor/Private/TextureEditorToolkit.cpp
@@ -1141,8 +1141,8 @@ void FTextureEditorToolkit::CreateInternalWidgets()
 		}
 		else
 		{
-			UE_LOG(LogTextureEditor, Warning, TEXT("No compression block size found in settings - using 256KB"));
-			CompressionBlockSize = 256*1024;
+			UE_LOG(LogTextureEditor, Verbose, TEXT("No compression block size found in settings - using 64KB"));
+			CompressionBlockSize = 64*1024;
 		}
 	}
 
