TextureDetails remove PowerOfTwo customization
was messing up property value change detection and served no purpose

diff --git a/Engine/Source/Editor/TextureEditor/Private/Customizations/TextureDetailsCustomization.cpp b/Engine/Source/Editor/TextureEditor/Private/Customizations/TextureDetailsCustomization.cpp
--- a/Engine/Source/Editor/TextureEditor/Private/Customizations/TextureDetailsCustomization.cpp
+++ b/Engine/Source/Editor/TextureEditor/Private/Customizations/TextureDetailsCustomization.cpp
@@ -37,7 +37,6 @@ void FTextureDetails::CustomizeDetails(IDetailLayoutBuilder& DetailBuilder)
 	DetailBuilder.EditCategory("File Path");
 
 	MaxTextureSizePropertyHandle = DetailBuilder.GetProperty(GET_MEMBER_NAME_CHECKED(UTexture, MaxTextureSize));
-	PowerOfTwoModePropertyHandle = DetailBuilder.GetProperty(GET_MEMBER_NAME_CHECKED(UTexture, PowerOfTwoMode));
 	VirtualTextureStreamingPropertyHandle = DetailBuilder.GetProperty(GET_MEMBER_NAME_CHECKED(UTexture, VirtualTextureStreaming));
 
 	// Customize MaxTextureSize
@@ -84,47 +83,6 @@ void FTextureDetails::CustomizeDetails(IDetailLayoutBuilder& DetailBuilder)
 			];
 	}
 
-	// Customize PowerOfTwoMode
-	if( PowerOfTwoModePropertyHandle->IsValidHandle() )
-	{
-		IDetailCategoryBuilder& TextureCategory = DetailBuilder.EditCategory("Texture");
-		IDetailPropertyRow& PowerOfTwoModePropertyRow = TextureCategory.AddProperty(GET_MEMBER_NAME_CHECKED(UTexture, PowerOfTwoMode));
-
-		PowerOfTwoModePropertyHandle->SetOnPropertyResetToDefault(FSimpleDelegate::CreateSP(this, &FTextureDetails::OnPropertyResetToDefault));
-
-		// Generate a list of enum values for the combo box
-		TArray<FText> PowerOfTwoModeComboBoxToolTips;
-		TArray<bool> RestrictedList;
-		PowerOfTwoModePropertyHandle->GeneratePossibleValues(PowerOfTwoModeComboBoxList, PowerOfTwoModeComboBoxToolTips, RestrictedList);
-
-		uint8 PowerOfTwoMode;
-		ensure(PowerOfTwoModePropertyHandle->GetValue(PowerOfTwoMode) == FPropertyAccess::Success);
-
-		TSharedPtr<SWidget> NameWidget;
-		TSharedPtr<SWidget> ValueWidget;
-		FDetailWidgetRow Row;
-		PowerOfTwoModePropertyRow.GetDefaultWidgets(NameWidget, ValueWidget, Row);
-
-		const bool bShowChildren = true;
-		PowerOfTwoModePropertyRow.CustomWidget(bShowChildren)
-			.NameContent()
-			.MinDesiredWidth(Row.NameWidget.MinWidth)
-			.MaxDesiredWidth(Row.NameWidget.MaxWidth)
-			[
-				NameWidget.ToSharedRef()
-			]
-			.ValueContent()
-			.MinDesiredWidth(Row.ValueWidget.MinWidth)
-			.MaxDesiredWidth(Row.ValueWidget.MaxWidth)
-			[
-				SAssignNew(TextComboBox, STextComboBox)
-				.Font(IDetailLayoutBuilder::GetDetailFont())
-				.OptionsSource(&PowerOfTwoModeComboBoxList)
-				.InitiallySelectedItem(PowerOfTwoModeComboBoxList[PowerOfTwoMode])
-				.OnSelectionChanged(this, &FTextureDetails::OnPowerOfTwoModeChanged)
-			];
-	}
-
 	// Hide the option to enable VT streaming, if VT is disabled for the project
 	if (VirtualTextureStreamingPropertyHandle.IsValid())
 	{
@@ -196,21 +154,5 @@ void FTextureDetails::OnEndSliderMovement(int32 NewValue)
 	GEditor->EndTransaction();
 }
 
-void FTextureDetails::OnPropertyResetToDefault() const
-{
-	uint8 CurrentPowerOfTwoMode;
-	ensure(PowerOfTwoModePropertyHandle->GetValue(CurrentPowerOfTwoMode) == FPropertyAccess::Success);
-	TextComboBox->SetSelectedItem(PowerOfTwoModeComboBoxList[CurrentPowerOfTwoMode]);
-}
-
-void FTextureDetails::OnPowerOfTwoModeChanged(TSharedPtr<FString> NewValue, ESelectInfo::Type SelectInfo)
-{
-	int32 NewPowerOfTwoMode = PowerOfTwoModeComboBoxList.Find(NewValue);
-	check(NewPowerOfTwoMode != INDEX_NONE);
-
-	int32 PowerOfTwoMode = PowerOfTwoModeComboBoxList.Find(NewValue);
-	check(PowerOfTwoMode != INDEX_NONE);
-	ensure(PowerOfTwoModePropertyHandle->SetValue(static_cast<uint8>(PowerOfTwoMode)) == FPropertyAccess::Success);
-}
 
 #undef LOCTEXT_NAMESPACE
diff --git a/Engine/Source/Editor/TextureEditor/Private/Customizations/TextureDetailsCustomization.h b/Engine/Source/Editor/TextureEditor/Private/Customizations/TextureDetailsCustomization.h
--- a/Engine/Source/Editor/TextureEditor/Private/Customizations/TextureDetailsCustomization.h
+++ b/Engine/Source/Editor/TextureEditor/Private/Customizations/TextureDetailsCustomization.h
@@ -30,13 +30,9 @@ private:
 	void OnMaxTextureSizeCommitted(int32 NewValue, ETextCommit::Type CommitInfo);
 	void OnBeginSliderMovement();
 	void OnEndSliderMovement(int32 NewValue);
-	void OnPowerOfTwoModeChanged(TSharedPtr<FString> NewValue, ESelectInfo::Type SelectInfo);
-	void OnPropertyResetToDefault() const;
 
 	TSharedPtr<IPropertyHandle> MaxTextureSizePropertyHandle;
-	TSharedPtr<IPropertyHandle> PowerOfTwoModePropertyHandle;
 	TSharedPtr<IPropertyHandle> VirtualTextureStreamingPropertyHandle;
-	TArray<TSharedPtr<FString>> PowerOfTwoModeComboBoxList;
 	TWeakObjectPtr<UObject> TextureBeingCustomized;
 	TSharedPtr<STextComboBox> TextComboBox;
 	bool bIsUsingSlider;
