Texture ForceRecompressDDCUID to let you manually change the DDC key to get a rebuild

diff --git a/Engine/Source/Developer/TextureBuild/Private/TextureBuildFunction.cpp b/Engine/Source/Developer/TextureBuild/Private/TextureBuildFunction.cpp
--- a/Engine/Source/Developer/TextureBuild/Private/TextureBuildFunction.cpp
+++ b/Engine/Source/Developer/TextureBuild/Private/TextureBuildFunction.cpp
@@ -95,6 +95,7 @@ static FTextureBuildSettings ReadBuildSettingsFromCompactBinary(const FCbObjectV
 	FTextureBuildSettings BuildSettings;
 	BuildSettings.FormatConfigOverride = Object["FormatConfigOverride"].AsObjectView();
 	FCbObjectView ColorAdjustmentCbObj = Object["ColorAdjustment"].AsObjectView();
+	ReadCbField(ColorAdjustmentCbObj["ForceRecompressDDCUID"], BuildSettings.ForceRecompressDDCUID);
 	FColorAdjustmentParameters& ColorAdjustment = BuildSettings.ColorAdjustment;
 	ReadCbField(ColorAdjustmentCbObj["AdjustBrightness"], ColorAdjustment.AdjustBrightness);
 	ReadCbField(ColorAdjustmentCbObj["AdjustBrightnessCurve"], ColorAdjustment.AdjustBrightnessCurve);
diff --git a/Engine/Source/Developer/TextureCompressor/Public/TextureCompressorModule.h b/Engine/Source/Developer/TextureCompressor/Public/TextureCompressorModule.h
--- a/Engine/Source/Developer/TextureCompressor/Public/TextureCompressorModule.h
+++ b/Engine/Source/Developer/TextureCompressor/Public/TextureCompressorModule.h
@@ -70,6 +70,8 @@ struct FColorAdjustmentParameters
  */
 struct FTextureBuildSettings
 {
+	/** Force recompress by changing DDC key */
+	uint32 ForceRecompressDDCUID;
 	/** Format specific config object view or null if no format specific config is applied as part of this build. */
 	FCbObjectView FormatConfigOverride;
 	/** Color adjustment parameters. */
diff --git a/Engine/Source/Editor/TextureEditor/Private/Customizations/TextureDetailsCustomization.cpp b/Engine/Source/Editor/TextureEditor/Private/Customizations/TextureDetailsCustomization.cpp
--- a/Engine/Source/Editor/TextureEditor/Private/Customizations/TextureDetailsCustomization.cpp
+++ b/Engine/Source/Editor/TextureEditor/Private/Customizations/TextureDetailsCustomization.cpp
@@ -11,7 +11,7 @@
 #include "IDetailPropertyRow.h"
 #include "DetailCategoryBuilder.h"
 #include "Widgets/Input/SNumericEntryBox.h"
-#include "Widgets/Input/STextComboBox.h"
+#include "Widgets/Input/SButton.h"
 
 #define LOCTEXT_NAMESPACE "FTextureDetails"
 
@@ -36,9 +36,58 @@ void FTextureDetails::CustomizeDetails(IDetailLayoutBuilder& DetailBuilder)
 	DetailBuilder.EditCategory("Adjustments");
 	DetailBuilder.EditCategory("File Path");
 
+	ForceRecompressDDCUIDPropertyHandle = DetailBuilder.GetProperty(GET_MEMBER_NAME_CHECKED(UTexture, ForceRecompressDDCUID));
 	MaxTextureSizePropertyHandle = DetailBuilder.GetProperty(GET_MEMBER_NAME_CHECKED(UTexture, MaxTextureSize));
 	VirtualTextureStreamingPropertyHandle = DetailBuilder.GetProperty(GET_MEMBER_NAME_CHECKED(UTexture, VirtualTextureStreaming));
+		
+	if( ForceRecompressDDCUIDPropertyHandle->IsValidHandle() )
+	{
+		IDetailCategoryBuilder& CompressionCategory = DetailBuilder.EditCategory("Compression");
+		IDetailPropertyRow& ForceRecompressDDCUIDPropertyRow = CompressionCategory.AddProperty(GET_MEMBER_NAME_CHECKED(UTexture, ForceRecompressDDCUID));
+		TSharedPtr<SWidget> NameWidget;
+		TSharedPtr<SWidget> ValueWidget;
+		FDetailWidgetRow Row;
+		ForceRecompressDDCUIDPropertyRow.GetDefaultWidgets(NameWidget, ValueWidget, Row);
 
+		const bool bShowChildren = true;
+		ForceRecompressDDCUIDPropertyRow.CustomWidget(bShowChildren)
+			.NameContent()
+			.MinDesiredWidth(Row.NameWidget.MinWidth)
+			.MaxDesiredWidth(Row.NameWidget.MaxWidth)
+			[
+				NameWidget.ToSharedRef()
+			]
+			.ValueContent()
+			.MinDesiredWidth(Row.ValueWidget.MinWidth)
+			.MaxDesiredWidth(Row.ValueWidget.MaxWidth)
+			.VAlign(VAlign_Fill)
+			.HAlign(HAlign_Fill)
+			[
+				SNew(SHorizontalBox)
+				+ SHorizontalBox::Slot()
+				[
+					ValueWidget.ToSharedRef()
+				]
+				+ SHorizontalBox::Slot()
+				.HAlign(HAlign_Right)
+				.VAlign(VAlign_Center)
+				.AutoWidth()
+				[
+					SNew(SButton)
+					.OnClicked(this, &FTextureDetails::OnForceRecompressDDCUIDClicked)
+					.ContentPadding(FMargin(2))
+					.Content()
+					[
+						SNew(STextBlock)
+						.Justification(ETextJustify::Center)
+						.Font(IDetailLayoutBuilder::GetDetailFont())
+						.Text(LOCTEXT("ForceRecompressDDCUIDRandom", "Random"))
+						.ToolTipText(LOCTEXT("ForceRecompressDDCUIDRandomTooltip", "Generate a random UID"))
+					]
+				]
+			];
+	}
+	
 	// Customize MaxTextureSize
 	if( MaxTextureSizePropertyHandle->IsValidHandle() )
 	{
@@ -53,6 +102,7 @@ void FTextureDetails::CustomizeDetails(IDetailLayoutBuilder& DetailBuilder)
 
 		if (UTexture* Texture = Cast<UTexture>(TextureBeingCustomized.Get()))
 		{
+			// GetMaximumDimension is for current RHI and texture type
 			MaxTextureSize = Texture->GetMaximumDimension();
 		}
 
@@ -95,6 +145,26 @@ void FTextureDetails::CustomizeDetails(IDetailLayoutBuilder& DetailBuilder)
 	}
 }
 
+FReply FTextureDetails::OnForceRecompressDDCUIDClicked()
+{
+	UTexture* Texture = Cast<UTexture>(TextureBeingCustomized.Get());
+	if ( Texture == nullptr )
+	{
+		return FReply::Unhandled();
+	}
+
+	// get a good random value :
+	FGuid Guid = FGuid::NewGuid();
+	uint32 RandomValue = Guid.A ^ Guid.B ^ Guid.C ^ Guid.D;
+
+	// don't just store the value, call SetValue so you get PostEditChange, etc.
+	//Texture->ForceRecompressDDCUID = RandomValue;
+	check( ForceRecompressDDCUIDPropertyHandle );
+	ForceRecompressDDCUIDPropertyHandle->SetValue( RandomValue );
+
+	return FReply::Handled();
+}
+
 /** @return The value or unset if properties with multiple values are viewed */
 TOptional<int32> FTextureDetails::OnGetMaxTextureSize() const
 {
diff --git a/Engine/Source/Editor/TextureEditor/Private/Customizations/TextureDetailsCustomization.h b/Engine/Source/Editor/TextureEditor/Private/Customizations/TextureDetailsCustomization.h
--- a/Engine/Source/Editor/TextureEditor/Private/Customizations/TextureDetailsCustomization.h
+++ b/Engine/Source/Editor/TextureEditor/Private/Customizations/TextureDetailsCustomization.h
@@ -7,6 +7,7 @@
 #include "UObject/WeakObjectPtr.h"
 #include "Types/SlateEnums.h"
 #include "IDetailCustomization.h"
+#include "Input/Reply.h"
 
 class IDetailLayoutBuilder;
 class IPropertyHandle;
@@ -30,11 +31,12 @@ private:
 	void OnMaxTextureSizeCommitted(int32 NewValue, ETextCommit::Type CommitInfo);
 	void OnBeginSliderMovement();
 	void OnEndSliderMovement(int32 NewValue);
+	FReply OnForceRecompressDDCUIDClicked();
 
+	TSharedPtr<IPropertyHandle> ForceRecompressDDCUIDPropertyHandle;
 	TSharedPtr<IPropertyHandle> MaxTextureSizePropertyHandle;
 	TSharedPtr<IPropertyHandle> VirtualTextureStreamingPropertyHandle;
 	TWeakObjectPtr<UObject> TextureBeingCustomized;
-	TSharedPtr<STextComboBox> TextComboBox;
 	bool bIsUsingSlider;
 };
 
diff --git a/Engine/Source/Runtime/Engine/Classes/Engine/Texture.h b/Engine/Source/Runtime/Engine/Classes/Engine/Texture.h
--- a/Engine/Source/Runtime/Engine/Classes/Engine/Texture.h
+++ b/Engine/Source/Runtime/Engine/Classes/Engine/Texture.h
@@ -1160,7 +1160,7 @@ public:
 	/** If enabled, defer compression of the texture until save or manually compressed in the texture editor. */
 	UPROPERTY(EditAnywhere, BlueprintReadWrite, Transient, Category=Compression, meta=(NoResetToDefault), meta=(DisplayName="Editor Defer Compression"))
 	uint32 DeferCompression:1;
-
+	
 	/** How aggressively should any relevant lossy compression be applied. For compressors that support EncodeSpeed (i.e. Oodle), this is only
 	*	applied if enabled (see Project Settings -> Texture Encoding). Note that this is *in addition* to any
 	*	unavoidable loss due to the target format - selecting "No Lossy Compression" will not result in zero distortion for BCn formats.
@@ -1169,7 +1169,7 @@ public:
 	TEnumAsByte<ETextureLossyCompressionAmount> LossyCompressionAmount;
 	
 	/** Oodle Texture SDK Version to encode with.  Enter 'latest' to update; 'None' preserves legacy encoding to avoid patches. */
-	UPROPERTY(EditAnywhere, BlueprintReadWrite, Category = Compression, AdvancedDisplay)
+	UPROPERTY(EditAnywhere, BlueprintReadWrite, Category = Compression, AdvancedDisplay, meta=(NoResetToDefault))
 	FName OodleTextureSdkVersion;
 
 	/** The maximum resolution for generated textures. A value of 0 means the maximum size for the format on each platform. */
@@ -1179,6 +1179,10 @@ public:
 	/** The compression quality for generated ASTC textures (i.e. mobile platform textures). */
 	UPROPERTY(EditAnywhere, BlueprintReadWrite, Category=Compression, meta = (DisplayName = "ASTC Compression Quality"), AdvancedDisplay)
 	TEnumAsByte<enum ETextureCompressionQuality> CompressionQuality;
+	
+	/** Increment ForceRecompressDDCUID to force recompression (change the DDC key) */
+	UPROPERTY(EditAnywhere, Category=Compression, meta=(DisplayName="Force Recompress UID"), AdvancedDisplay, meta=(NoResetToDefault))
+	uint32 ForceRecompressDDCUID;
 
 	/** Removed. */
 	UPROPERTY()
diff --git a/Engine/Source/Runtime/Engine/Private/Texture.cpp b/Engine/Source/Runtime/Engine/Private/Texture.cpp
--- a/Engine/Source/Runtime/Engine/Private/Texture.cpp
+++ b/Engine/Source/Runtime/Engine/Private/Texture.cpp
@@ -126,6 +126,7 @@ UTexture::UTexture(const FObjectInitializer& ObjectInitializer)
 	CompressionNone = 0;
 	CompressFinal = 0;
 	DeferCompression = 0;
+	ForceRecompressDDCUID = 0;
 	LossyCompressionAmount = TLCA_Default;
 	MaxTextureSize = 0; // means no limitation
 	MipGenSettings = TMGS_FromTextureGroup;
diff --git a/Engine/Source/Runtime/Engine/Private/TextureDerivedData.cpp b/Engine/Source/Runtime/Engine/Private/TextureDerivedData.cpp
--- a/Engine/Source/Runtime/Engine/Private/TextureDerivedData.cpp
+++ b/Engine/Source/Runtime/Engine/Private/TextureDerivedData.cpp
@@ -64,6 +64,7 @@
 // In case of merge conflicts with DDC versions, you MUST generate a new GUID and set this new
 // guid as version
 // this is put in the DDC1 and the DDC2 key
+// @@!! next time this changes clean up SerializeForKey todo marks
 #define TEXTURE_DERIVEDDATA_VER		TEXT("95BCE5A0BFB949539A18684748C633C9")
 
 // This GUID is mixed into DDC version for virtual textures only, this allows updating DDC version for VT without invalidating DDC for all textures
@@ -206,6 +207,7 @@ static void SerializeForKey(FArchive& Ar, const FTextureBuildSettings& Settings)
 
 	// weird/unncessary way to do this :
 	// should have just serialized the bool bForceAlphaChannel
+	// @@!! these can go away whenever we bump the overall ddc key
 	if (Settings.bForceAlphaChannel)
 	{
 		TempGuid = FGuid(0x2C9DF7E3, 0xBC9D413B, 0xBF963C7A, 0x3F27E8B1);
@@ -215,6 +217,7 @@ static void SerializeForKey(FArchive& Ar, const FTextureBuildSettings& Settings)
 	if ( Settings.MaxTextureResolution != FTextureBuildSettings::MaxTextureResolutionDefault &&
 		( Settings.MipGenSettings == TMGS_LeaveExistingMips || Settings.bDoScaleMipsForAlphaCoverage ) )
 	{
+		// @@!! these can go away whenever we bump the overall ddc key
 		// behavior of MaxTextureResolution + LeaveExistingMips or bDoScaleMipsForAlphaCoverage changed, so modify the key :
 		TempGuid = FGuid(0x418B8584, 0x72D54EA5, 0xBA8E8C2B, 0xECC880DE);
 		Ar << TempGuid;
@@ -222,22 +225,30 @@ static void SerializeForKey(FArchive& Ar, const FTextureBuildSettings& Settings)
 
 	if ( Settings.bVolume )
 	{
+		// @@!! these can go away whenever we bump the overall ddc key
 		TempGuid = FGuid(0xCC4348B8,0x84714993,0xAB1E2C93,0x8EA6C9E0);
 		Ar << TempGuid;
 	}
 
 	if ( Settings.bVirtualStreamable && Settings.bSRGB && Settings.bUseLegacyGamma )
 	{
+		// @@!! these can go away whenever we bump the overall ddc key
 		TempGuid = FGuid(0xCAEDDFB6,0xEDC2455D,0x8D45B90C,0x3A1B7783);
 		Ar << TempGuid;
 	}
 
 	// do not change key if old mip filter is used for old textures
+	// @@!! these can go away whenever we bump the overall ddc key
 	if (Settings.bUseNewMipFilter)
 	{
 		TempGuid = FGuid(0x27B79A99, 0xE1A5458E, 0xAB619475, 0xCD01AD2A);
 		Ar << TempGuid;
 	}
+	
+	if ( Settings.ForceRecompressDDCUID != 0 )
+	{
+		TempUint32 = Settings.ForceRecompressDDCUID; Ar << TempUint32;
+	}
 
 	// Note - compression quality is added to the DDC by the formats (based on whether they
 	// use them or not).
@@ -272,7 +283,9 @@ void GetTextureDerivedDataKeySuffix(const UTexture& Texture, const FTextureBuild
 		{
 			Version = TextureFormat->GetVersion(BuildSettings.TextureFormatName, &BuildSettings);
 		}
+		// @@!! else error !?
 	}
+	// @@!! else error !?
 	
 	FString CompositeTextureStr;
 
@@ -744,6 +757,7 @@ static void GetTextureBuildSettings(
 		OutBuildSettings.AlphaCoverageThresholds = FVector4f(0,0,0,0);
 	}
 
+	OutBuildSettings.ForceRecompressDDCUID = Texture.ForceRecompressDDCUID;
 	OutBuildSettings.bUseNewMipFilter = Texture.bUseNewMipFilter;
 	OutBuildSettings.bComputeBokehAlpha = (Texture.LODGroup == TEXTUREGROUP_Bokeh);
 	OutBuildSettings.bReplicateAlpha = false;
diff --git a/Engine/Source/Runtime/Engine/Private/TextureDerivedDataBuildUtils.cpp b/Engine/Source/Runtime/Engine/Private/TextureDerivedDataBuildUtils.cpp
--- a/Engine/Source/Runtime/Engine/Private/TextureDerivedDataBuildUtils.cpp
+++ b/Engine/Source/Runtime/Engine/Private/TextureDerivedDataBuildUtils.cpp
@@ -120,6 +120,7 @@ static void WriteBuildSettings(FCbWriter& Writer, const FTextureBuildSettings& B
 		// AlphaCoverageThresholds do not affect build if bDoScaleMipsForAlphaCoverage is off
 		WriteCbFieldWithDefault(Writer, "AlphaCoverageThresholds", BuildSettings.AlphaCoverageThresholds, DefaultSettings.AlphaCoverageThresholds);
 	}
+	WriteCbFieldWithDefault(Writer, "ForceRecompressDDCUID", BuildSettings.ForceRecompressDDCUID, DefaultSettings.ForceRecompressDDCUID);
 	WriteCbFieldWithDefault<bool>(Writer, "bUseNewMipFilter", BuildSettings.bUseNewMipFilter, DefaultSettings.bUseNewMipFilter);
 	WriteCbFieldWithDefault(Writer, "MipSharpening", BuildSettings.MipSharpening, DefaultSettings.MipSharpening);
 	WriteCbFieldWithDefault(Writer, "DiffuseConvolveMipLevel", BuildSettings.DiffuseConvolveMipLevel, DefaultSettings.DiffuseConvolveMipLevel);
