Texture: Renamed ForceRecompressDDCUID to CompressionCacheId

diff --git a/Engine/Source/Developer/TextureBuild/Private/TextureBuildFunction.cpp b/Engine/Source/Developer/TextureBuild/Private/TextureBuildFunction.cpp
--- a/Engine/Source/Developer/TextureBuild/Private/TextureBuildFunction.cpp
+++ b/Engine/Source/Developer/TextureBuild/Private/TextureBuildFunction.cpp
@@ -49,6 +49,7 @@ static void ReadCbField(FCbFieldView Field, int32& OutValue) { OutValue = Field.
 static void ReadCbField(FCbFieldView Field, uint8& OutValue) { OutValue = Field.AsUInt8(OutValue); }
 static void ReadCbField(FCbFieldView Field, uint32& OutValue) { OutValue = Field.AsUInt32(OutValue); }
 static void ReadCbField(FCbFieldView Field, float& OutValue) { OutValue = Field.AsFloat(OutValue); }
+static void ReadCbField(FCbFieldView Field, FGuid& OutValue) { OutValue = Field.AsUuid(); }
 
 static void ReadCbField(FCbFieldView Field, FName& OutValue)
 {
@@ -95,7 +96,7 @@ static FTextureBuildSettings ReadBuildSettingsFromCompactBinary(const FCbObjectV
 	FTextureBuildSettings BuildSettings;
 	BuildSettings.FormatConfigOverride = Object["FormatConfigOverride"].AsObjectView();
 	FCbObjectView ColorAdjustmentCbObj = Object["ColorAdjustment"].AsObjectView();
-	ReadCbField(ColorAdjustmentCbObj["ForceRecompressDDCUID"], BuildSettings.ForceRecompressDDCUID);
+	ReadCbField(ColorAdjustmentCbObj["CompressionCacheId"], BuildSettings.CompressionCacheId);
 	FColorAdjustmentParameters& ColorAdjustment = BuildSettings.ColorAdjustment;
 	ReadCbField(ColorAdjustmentCbObj["AdjustBrightness"], ColorAdjustment.AdjustBrightness);
 	ReadCbField(ColorAdjustmentCbObj["AdjustBrightnessCurve"], ColorAdjustment.AdjustBrightnessCurve);
diff --git a/Engine/Source/Developer/TextureCompressor/Public/TextureCompressorModule.h b/Engine/Source/Developer/TextureCompressor/Public/TextureCompressorModule.h
--- a/Engine/Source/Developer/TextureCompressor/Public/TextureCompressorModule.h
+++ b/Engine/Source/Developer/TextureCompressor/Public/TextureCompressorModule.h
@@ -70,8 +70,8 @@ struct FColorAdjustmentParameters
  */
 struct FTextureBuildSettings
 {
-	/** Force recompress by changing DDC key */
-	uint32 ForceRecompressDDCUID;
+	/** An optional ID to change the cache key. */
+	FGuid CompressionCacheId;
 	/** Format specific config object view or null if no format specific config is applied as part of this build. */
 	FCbObjectView FormatConfigOverride;
 	/** Color adjustment parameters. */
@@ -217,8 +217,7 @@ struct FTextureBuildSettings
 
 	/** Default settings. */
 	FTextureBuildSettings()
-		: ForceRecompressDDCUID(0)
-		, bDoScaleMipsForAlphaCoverage(false)
+		: bDoScaleMipsForAlphaCoverage(false)
 		, AlphaCoverageThresholds(0, 0, 0, 0)
 		, bUseNewMipFilter(false)
 		, MipSharpening(0.0f)
diff --git a/Engine/Source/Editor/TextureEditor/Private/Customizations/TextureDetailsCustomization.cpp b/Engine/Source/Editor/TextureEditor/Private/Customizations/TextureDetailsCustomization.cpp
--- a/Engine/Source/Editor/TextureEditor/Private/Customizations/TextureDetailsCustomization.cpp
+++ b/Engine/Source/Editor/TextureEditor/Private/Customizations/TextureDetailsCustomization.cpp
@@ -4,14 +4,12 @@
 #include "Misc/MessageDialog.h"
 #include "Widgets/DeclarativeSyntaxSupport.h"
 #include "Engine/Texture.h"
-//#include "Engine/Texture2D.h"
 #include "Editor.h"
 #include "DetailLayoutBuilder.h"
 #include "DetailWidgetRow.h"
 #include "IDetailPropertyRow.h"
 #include "DetailCategoryBuilder.h"
 #include "Widgets/Input/SNumericEntryBox.h"
-#include "Widgets/Input/SButton.h"
 
 #define LOCTEXT_NAMESPACE "FTextureDetails"
 
@@ -36,58 +34,9 @@ void FTextureDetails::CustomizeDetails(IDetailLayoutBuilder& DetailBuilder)
 	DetailBuilder.EditCategory("Adjustments");
 	DetailBuilder.EditCategory("File Path");
 
-	ForceRecompressDDCUIDPropertyHandle = DetailBuilder.GetProperty(GET_MEMBER_NAME_CHECKED(UTexture, ForceRecompressDDCUID));
 	MaxTextureSizePropertyHandle = DetailBuilder.GetProperty(GET_MEMBER_NAME_CHECKED(UTexture, MaxTextureSize));
 	VirtualTextureStreamingPropertyHandle = DetailBuilder.GetProperty(GET_MEMBER_NAME_CHECKED(UTexture, VirtualTextureStreaming));
-		
-	if( ForceRecompressDDCUIDPropertyHandle->IsValidHandle() )
-	{
-		IDetailCategoryBuilder& CompressionCategory = DetailBuilder.EditCategory("Compression");
-		IDetailPropertyRow& ForceRecompressDDCUIDPropertyRow = CompressionCategory.AddProperty(GET_MEMBER_NAME_CHECKED(UTexture, ForceRecompressDDCUID));
-		TSharedPtr<SWidget> NameWidget;
-		TSharedPtr<SWidget> ValueWidget;
-		FDetailWidgetRow Row;
-		ForceRecompressDDCUIDPropertyRow.GetDefaultWidgets(NameWidget, ValueWidget, Row);
 
-		const bool bShowChildren = true;
-		ForceRecompressDDCUIDPropertyRow.CustomWidget(bShowChildren)
-			.NameContent()
-			.MinDesiredWidth(Row.NameWidget.MinWidth)
-			.MaxDesiredWidth(Row.NameWidget.MaxWidth)
-			[
-				NameWidget.ToSharedRef()
-			]
-			.ValueContent()
-			.MinDesiredWidth(Row.ValueWidget.MinWidth)
-			.MaxDesiredWidth(Row.ValueWidget.MaxWidth)
-			.VAlign(VAlign_Fill)
-			.HAlign(HAlign_Fill)
-			[
-				SNew(SHorizontalBox)
-				+ SHorizontalBox::Slot()
-				[
-					ValueWidget.ToSharedRef()
-				]
-				+ SHorizontalBox::Slot()
-				.HAlign(HAlign_Right)
-				.VAlign(VAlign_Center)
-				.AutoWidth()
-				[
-					SNew(SButton)
-					.OnClicked(this, &FTextureDetails::OnForceRecompressDDCUIDClicked)
-					.ContentPadding(FMargin(2))
-					.Content()
-					[
-						SNew(STextBlock)
-						.Justification(ETextJustify::Center)
-						.Font(IDetailLayoutBuilder::GetDetailFont())
-						.Text(LOCTEXT("ForceRecompressDDCUIDRandom", "Random"))
-						.ToolTipText(LOCTEXT("ForceRecompressDDCUIDRandomTooltip", "Generate a random UID"))
-					]
-				]
-			];
-	}
-	
 	// Customize MaxTextureSize
 	if( MaxTextureSizePropertyHandle->IsValidHandle() )
 	{
@@ -145,26 +94,6 @@ void FTextureDetails::CustomizeDetails(IDetailLayoutBuilder& DetailBuilder)
 	}
 }
 
-FReply FTextureDetails::OnForceRecompressDDCUIDClicked()
-{
-	UTexture* Texture = Cast<UTexture>(TextureBeingCustomized.Get());
-	if ( Texture == nullptr )
-	{
-		return FReply::Unhandled();
-	}
-
-	// get a good random value :
-	FGuid Guid = FGuid::NewGuid();
-	uint32 RandomValue = Guid.A ^ Guid.B ^ Guid.C ^ Guid.D;
-
-	// don't just store the value, call SetValue so you get PostEditChange, etc.
-	//Texture->ForceRecompressDDCUID = RandomValue;
-	check( ForceRecompressDDCUIDPropertyHandle );
-	ForceRecompressDDCUIDPropertyHandle->SetValue( RandomValue );
-
-	return FReply::Handled();
-}
-
 /** @return The value or unset if properties with multiple values are viewed */
 TOptional<int32> FTextureDetails::OnGetMaxTextureSize() const
 {
diff --git a/Engine/Source/Editor/TextureEditor/Private/Customizations/TextureDetailsCustomization.h b/Engine/Source/Editor/TextureEditor/Private/Customizations/TextureDetailsCustomization.h
--- a/Engine/Source/Editor/TextureEditor/Private/Customizations/TextureDetailsCustomization.h
+++ b/Engine/Source/Editor/TextureEditor/Private/Customizations/TextureDetailsCustomization.h
@@ -7,7 +7,6 @@
 #include "UObject/WeakObjectPtr.h"
 #include "Types/SlateEnums.h"
 #include "IDetailCustomization.h"
-#include "Input/Reply.h"
 
 class IDetailLayoutBuilder;
 class IPropertyHandle;
@@ -31,9 +30,7 @@ private:
 	void OnMaxTextureSizeCommitted(int32 NewValue, ETextCommit::Type CommitInfo);
 	void OnBeginSliderMovement();
 	void OnEndSliderMovement(int32 NewValue);
-	FReply OnForceRecompressDDCUIDClicked();
 
-	TSharedPtr<IPropertyHandle> ForceRecompressDDCUIDPropertyHandle;
 	TSharedPtr<IPropertyHandle> MaxTextureSizePropertyHandle;
 	TSharedPtr<IPropertyHandle> VirtualTextureStreamingPropertyHandle;
 	TWeakObjectPtr<UObject> TextureBeingCustomized;
diff --git a/Engine/Source/Runtime/Engine/Classes/Engine/Texture.h b/Engine/Source/Runtime/Engine/Classes/Engine/Texture.h
--- a/Engine/Source/Runtime/Engine/Classes/Engine/Texture.h
+++ b/Engine/Source/Runtime/Engine/Classes/Engine/Texture.h
@@ -1179,10 +1179,10 @@ public:
 	/** The compression quality for generated ASTC textures (i.e. mobile platform textures). */
 	UPROPERTY(EditAnywhere, BlueprintReadWrite, Category=Compression, meta = (DisplayName = "ASTC Compression Quality"), AdvancedDisplay)
 	TEnumAsByte<enum ETextureCompressionQuality> CompressionQuality;
-	
-	/** Increment ForceRecompressDDCUID to force recompression (change the DDC key) */
-	UPROPERTY(EditAnywhere, Category=Compression, meta=(DisplayName="Force Recompress UID"), AdvancedDisplay, meta=(NoResetToDefault))
-	uint32 ForceRecompressDDCUID;
+
+	/** Change this optional ID to force the texture to be recompressed by changing its cache key. */
+	UPROPERTY(EditAnywhere, Category=Compression, meta=(NoResetToDefault), meta=(DisplayName = "Compression Cache ID"), AdvancedDisplay)
+	FGuid CompressionCacheId;
 
 	/** Removed. */
 	UPROPERTY()
diff --git a/Engine/Source/Runtime/Engine/Private/Texture.cpp b/Engine/Source/Runtime/Engine/Private/Texture.cpp
--- a/Engine/Source/Runtime/Engine/Private/Texture.cpp
+++ b/Engine/Source/Runtime/Engine/Private/Texture.cpp
@@ -126,7 +126,7 @@ UTexture::UTexture(const FObjectInitializer& ObjectInitializer)
 	CompressionNone = 0;
 	CompressFinal = 0;
 	DeferCompression = 0;
-	ForceRecompressDDCUID = 0;
+	CompressionCacheId = FGuid(0, 0, 0, 0);
 	LossyCompressionAmount = TLCA_Default;
 	MaxTextureSize = 0; // means no limitation
 	MipGenSettings = TMGS_FromTextureGroup;
diff --git a/Engine/Source/Runtime/Engine/Private/TextureDerivedData.cpp b/Engine/Source/Runtime/Engine/Private/TextureDerivedData.cpp
--- a/Engine/Source/Runtime/Engine/Private/TextureDerivedData.cpp
+++ b/Engine/Source/Runtime/Engine/Private/TextureDerivedData.cpp
@@ -242,9 +242,9 @@ static void SerializeForKey(FArchive& Ar, const FTextureBuildSettings& Settings)
 		Ar << TempGuid;
 	}
 	
-	if ( Settings.ForceRecompressDDCUID != 0 )
+	if (Settings.CompressionCacheId.IsValid())
 	{
-		TempUint32 = Settings.ForceRecompressDDCUID; Ar << TempUint32;
+		TempGuid = Settings.CompressionCacheId; Ar << TempGuid;
 	}
 
 	// Note - compression quality is added to the DDC by the formats (based on whether they
@@ -754,7 +754,7 @@ static void GetTextureBuildSettings(
 		OutBuildSettings.AlphaCoverageThresholds = FVector4f(0,0,0,0);
 	}
 
-	OutBuildSettings.ForceRecompressDDCUID = Texture.ForceRecompressDDCUID;
+	OutBuildSettings.CompressionCacheId = Texture.CompressionCacheId;
 	OutBuildSettings.bUseNewMipFilter = Texture.bUseNewMipFilter;
 	OutBuildSettings.bComputeBokehAlpha = (Texture.LODGroup == TEXTUREGROUP_Bokeh);
 	OutBuildSettings.bReplicateAlpha = false;
diff --git a/Engine/Source/Runtime/Engine/Private/TextureDerivedDataBuildUtils.cpp b/Engine/Source/Runtime/Engine/Private/TextureDerivedDataBuildUtils.cpp
--- a/Engine/Source/Runtime/Engine/Private/TextureDerivedDataBuildUtils.cpp
+++ b/Engine/Source/Runtime/Engine/Private/TextureDerivedDataBuildUtils.cpp
@@ -120,7 +120,7 @@ static void WriteBuildSettings(FCbWriter& Writer, const FTextureBuildSettings& B
 		// AlphaCoverageThresholds do not affect build if bDoScaleMipsForAlphaCoverage is off
 		WriteCbFieldWithDefault(Writer, "AlphaCoverageThresholds", BuildSettings.AlphaCoverageThresholds, DefaultSettings.AlphaCoverageThresholds);
 	}
-	WriteCbFieldWithDefault(Writer, "ForceRecompressDDCUID", BuildSettings.ForceRecompressDDCUID, DefaultSettings.ForceRecompressDDCUID);
+	WriteCbFieldWithDefault(Writer, "CompressionCacheId", BuildSettings.CompressionCacheId, DefaultSettings.CompressionCacheId);
 	WriteCbFieldWithDefault<bool>(Writer, "bUseNewMipFilter", BuildSettings.bUseNewMipFilter, DefaultSettings.bUseNewMipFilter);
 	WriteCbFieldWithDefault(Writer, "MipSharpening", BuildSettings.MipSharpening, DefaultSettings.MipSharpening);
 	WriteCbFieldWithDefault(Writer, "DiffuseConvolveMipLevel", BuildSettings.DiffuseConvolveMipLevel, DefaultSettings.DiffuseConvolveMipLevel);
