Use centered zoom for the texture editor viewport.

diff --git a/Engine/Source/Editor/TextureEditor/Private/Models/TextureEditorViewportClient.cpp b/Engine/Source/Editor/TextureEditor/Private/Models/TextureEditorViewportClient.cpp
--- a/Engine/Source/Editor/TextureEditor/Private/Models/TextureEditorViewportClient.cpp
+++ b/Engine/Source/Editor/TextureEditor/Private/Models/TextureEditorViewportClient.cpp
@@ -442,14 +442,16 @@ void FTextureEditorViewportClient::UpdateScrollBars()
 
 	float VRatio = GetViewportVerticalScrollBarRatio();
 	float HRatio = GetViewportHorizontalScrollBarRatio();
+	float VDistFromTop = Viewport->GetVerticalScrollBar()->DistanceFromTop();
 	float VDistFromBottom = Viewport->GetVerticalScrollBar()->DistanceFromBottom();
+	float HDistFromTop = Viewport->GetHorizontalScrollBar()->DistanceFromTop();
 	float HDistFromBottom = Viewport->GetHorizontalScrollBar()->DistanceFromBottom();
 
 	if (VRatio < 1.0f)
 	{
 		if (VDistFromBottom < 1.0f)
 		{
-			Viewport->GetVerticalScrollBar()->SetState(FMath::Clamp(1.0f - VRatio - VDistFromBottom, 0.0f, 1.0f), VRatio);
+			Viewport->GetVerticalScrollBar()->SetState(FMath::Clamp((1.0f + VDistFromTop - VDistFromBottom - VRatio) * 0.5f, 0.0f, 1.0f - VRatio), VRatio);
 		}
 		else
 		{
@@ -461,7 +463,7 @@ void FTextureEditorViewportClient::UpdateScrollBars()
 	{
 		if (HDistFromBottom < 1.0f)
 		{
-			Viewport->GetHorizontalScrollBar()->SetState(FMath::Clamp(1.0f - HRatio - HDistFromBottom, 0.0f, 1.0f), HRatio);
+			Viewport->GetHorizontalScrollBar()->SetState(FMath::Clamp((1.0f + HDistFromTop - HDistFromBottom - HRatio) * 0.5f, 0.0f, 1.0f - HRatio), HRatio);
 		}
 		else
 		{
@@ -480,14 +482,16 @@ FVector2D FTextureEditorViewportClient::GetViewportScrollBarPositions() const
 		UTexture* Texture = TextureEditorPtr.Pin()->GetTexture();
 		float VRatio = GetViewportVerticalScrollBarRatio();
 		float HRatio = GetViewportHorizontalScrollBarRatio();
+		float VDistFromTop = TextureEditorViewportPtr.Pin()->GetVerticalScrollBar()->DistanceFromTop();
 		float VDistFromBottom = TextureEditorViewportPtr.Pin()->GetVerticalScrollBar()->DistanceFromBottom();
+		float HDistFromTop = TextureEditorViewportPtr.Pin()->GetHorizontalScrollBar()->DistanceFromTop();
 		float HDistFromBottom = TextureEditorViewportPtr.Pin()->GetHorizontalScrollBar()->DistanceFromBottom();
 	
 		TextureEditorPtr.Pin()->CalculateTextureDimensions(Width, Height, Depth, ArraySize);
 
 		if ((TextureEditorViewportPtr.Pin()->GetVerticalScrollBar()->GetVisibility() == EVisibility::Visible) && VDistFromBottom < 1.0f)
 		{
-			Positions.Y = FMath::Clamp(1.0f - VRatio - VDistFromBottom, 0.0f, 1.0f) * Height;
+			Positions.Y = FMath::Clamp((1.0f + VDistFromTop - VDistFromBottom - VRatio) * 0.5f, 0.0f, 1.0f - VRatio) * Height;
 		}
 		else
 		{
@@ -496,7 +500,7 @@ FVector2D FTextureEditorViewportClient::GetViewportScrollBarPositions() const
 
 		if ((TextureEditorViewportPtr.Pin()->GetHorizontalScrollBar()->GetVisibility() == EVisibility::Visible) && HDistFromBottom < 1.0f)
 		{
-			Positions.X = FMath::Clamp(1.0f - HRatio - HDistFromBottom, 0.0f, 1.0f) * Width;
+			Positions.X = FMath::Clamp((1.0f + HDistFromTop - HDistFromBottom - HRatio) * 0.5f, 0.0f, 1.0f - HRatio) * Width;
 		}
 		else
 		{
