Calculate texture resource size based on the built platform data.

diff --git a/Engine/Source/Editor/TextureEditor/Private/TextureEditorToolkit.cpp b/Engine/Source/Editor/TextureEditor/Private/TextureEditorToolkit.cpp
--- a/Engine/Source/Editor/TextureEditor/Private/TextureEditorToolkit.cpp
+++ b/Engine/Source/Editor/TextureEditor/Private/TextureEditorToolkit.cpp
@@ -743,7 +743,7 @@ void FTextureEditorToolkit::PopulateQuickInfo( )
 	const uint32 MaxInGameDepth = FMath::Max<uint32>(SurfaceDepth >> MaxResMipBias, 1);
 
 	// Texture asset size
-	const uint32 Size = (Texture->GetResourceSizeBytes(EResourceSizeMode::Exclusive) + 512) / 1024;
+	const int64 ResourceSize = PlatformDataPtr && *PlatformDataPtr ? (*PlatformDataPtr)->GetPayloadSize() : Texture->GetResourceSizeBytes(EResourceSizeMode::Exclusive);
 
 	FNumberFormattingOptions SizeOptions;
 	SizeOptions.UseGrouping = false;
@@ -795,7 +795,7 @@ void FTextureEditorToolkit::PopulateQuickInfo( )
 		MaxInGameText->SetText(FText::Format( NSLOCTEXT("TextureEditor", "QuickInfo_MaxInGame_2x", "Max In-Game: {0}x{1}{2}"), FText::AsNumber(MaxInGameWidth, &Options), FText::AsNumber(MaxInGameHeight, &Options), CubemapAdd));
 	}
 
-	SizeText->SetText(FText::Format(NSLOCTEXT("TextureEditor", "QuickInfo_ResourceSize", "Resource Size: {0} KB"), FText::AsNumber(Size, &SizeOptions)));
+	SizeText->SetText(FText::Format(NSLOCTEXT("TextureEditor", "QuickInfo_ResourceSize", "Resource Size: {0} KB"), FText::AsNumber(FMath::DivideAndRoundNearest(ResourceSize, (int64)1024), &SizeOptions)));
 
 	FText Method = Texture->IsCurrentlyVirtualTextured() ? NSLOCTEXT("TextureEditor", "QuickInfo_MethodVirtualStreamed", "Virtual Streamed")
 													: (!Texture->IsStreamable() ? NSLOCTEXT("TextureEditor", "QuickInfo_MethodNotStreamed", "Not Streamed") 
diff --git a/Engine/Source/Runtime/Engine/Classes/Engine/Texture.h b/Engine/Source/Runtime/Engine/Classes/Engine/Texture.h
--- a/Engine/Source/Runtime/Engine/Classes/Engine/Texture.h
+++ b/Engine/Source/Runtime/Engine/Classes/Engine/Texture.h
@@ -998,6 +998,9 @@ public:
 	ENGINE_API int32 GetNumVTMips() const;
 	ENGINE_API EPixelFormat GetLayerPixelFormat(uint32 LayerIndex) const;
 
+	/** Return the size of the texture pixel data in bytes, not including headers or alignment. */
+	ENGINE_API int64 GetPayloadSize() const;
+
 private:
 
 	bool CanUseCookedDataPath() const;
diff --git a/Engine/Source/Runtime/Engine/Private/TextureDerivedData.cpp b/Engine/Source/Runtime/Engine/Private/TextureDerivedData.cpp
--- a/Engine/Source/Runtime/Engine/Private/TextureDerivedData.cpp
+++ b/Engine/Source/Runtime/Engine/Private/TextureDerivedData.cpp
@@ -2134,6 +2134,31 @@ EPixelFormat FTexturePlatformData::GetLayerPixelFormat(uint32 LayerIndex) const
 	return PixelFormat;
 }
 
+int64 FTexturePlatformData::GetPayloadSize() const
+{
+	int64 PayloadSize = 0;
+	if (VTData)
+	{
+		for (int32 ChunkIndex = 0; ChunkIndex < VTData->Chunks.Num(); ChunkIndex++)
+		{
+			PayloadSize += VTData->Chunks[ChunkIndex].SizeInBytes;
+		}
+	}
+	else
+	{
+		for (int32 MipIndex = 0; MipIndex < Mips.Num(); MipIndex++)
+		{
+			int32 BlockSizeX = FMath::DivideAndRoundUp(FMath::Max(Mips[MipIndex].SizeX, 1), GPixelFormats[PixelFormat].BlockSizeX);
+			int32 BlockSizeY = FMath::DivideAndRoundUp(FMath::Max(Mips[MipIndex].SizeY, 1), GPixelFormats[PixelFormat].BlockSizeY);
+			// for TextureCube and TextureCubeArray all the mipmaps contain the same number of slices, which is encoded in the PackedData member
+			// at the same time we can not just use SizeZ of a TextureCube mipmap, because for compatibility reasons it is always set to 1 and not 6 (which is the actual number of slices)
+			int32 BlockSizeZ = FMath::DivideAndRoundUp(FMath::Max(IsCubemap() ? GetNumSlices() : Mips[MipIndex].SizeZ, 1), GPixelFormats[PixelFormat].BlockSizeZ);
+			PayloadSize += (int64)GPixelFormats[PixelFormat].BlockBytes * BlockSizeX * BlockSizeY * BlockSizeZ;
+		}
+	}
+	return PayloadSize;
+}
+
 bool FTexturePlatformData::CanUseCookedDataPath() const
 {
 #if WITH_IOSTORE_IN_EDITOR
