Fix multiple texture editor issues: TextureCubeArray zoom, texture border, texture stats, tooltip information.

diff --git a/Engine/Source/Editor/TextureEditor/Private/Models/TextureEditorViewportClient.cpp b/Engine/Source/Editor/TextureEditor/Private/Models/TextureEditorViewportClient.cpp
--- a/Engine/Source/Editor/TextureEditor/Private/Models/TextureEditorViewportClient.cpp
+++ b/Engine/Source/Editor/TextureEditor/Private/Models/TextureEditorViewportClient.cpp
@@ -64,19 +64,20 @@ void FTextureEditorViewportClient::Draw(FViewport* Viewport, FCanvas* Canvas)
 	}
 
 	TSharedPtr<ITextureEditorToolkit> TextureEditorPinned = TextureEditorPtr.Pin();
-	
+	const UTextureEditorSettings& Settings = *GetDefault<UTextureEditorSettings>();
+
 	UTexture* Texture = TextureEditorPinned->GetTexture();
 	FVector2D Ratio = FVector2D(GetViewportHorizontalScrollBarRatio(), GetViewportVerticalScrollBarRatio());
 	FVector2D ViewportSize = FVector2D(TextureEditorViewportPtr.Pin()->GetViewport()->GetSizeXY());
 	FVector2D ScrollBarPos = GetViewportScrollBarPositions();
-	int32 YOffset = (Ratio.Y > 1.0f)? ((ViewportSize.Y - (ViewportSize.Y / Ratio.Y)) * 0.5f): 0;
-	int32 YPos = YOffset - ScrollBarPos.Y;
-	int32 XOffset = (Ratio.X > 1.0f)? ((ViewportSize.X - (ViewportSize.X / Ratio.X)) * 0.5f): 0;
-	int32 XPos = XOffset - ScrollBarPos.X;
+	int32 BorderSize = Settings.TextureBorderEnabled ? 1 : 0;
+	float YOffset = (Ratio.Y > 1.0f) ? ((ViewportSize.Y - (ViewportSize.Y / Ratio.Y)) * 0.5f) : 0;
+	int32 YPos = FMath::RoundToInt(YOffset - ScrollBarPos.Y + BorderSize);
+	float XOffset = (Ratio.X > 1.0f) ? ((ViewportSize.X - (ViewportSize.X / Ratio.X)) * 0.5f) : 0;
+	int32 XPos = FMath::RoundToInt(XOffset - ScrollBarPos.X + BorderSize);
 	
 	UpdateScrollBars();
 
-	const UTextureEditorSettings& Settings = *GetDefault<UTextureEditorSettings>();
 
 	Canvas->Clear( Settings.BackgroundColor );
 
@@ -97,8 +98,8 @@ void FTextureEditorViewportClient::Draw(FViewport* Viewport, FCanvas* Canvas)
 	TextureEditorPinned->PopulateQuickInfo();
 
 	// Figure out the size we need
-	uint32 Width, Height, Depth, ArraySize;;
-	TextureEditorPinned->CalculateTextureDimensions(Width, Height, Depth, ArraySize);
+	int32 Width, Height, Depth, ArraySize;
+	TextureEditorPinned->CalculateTextureDimensions(Width, Height, Depth, ArraySize, false);
 	const float MipLevel = (float)TextureEditorPinned->GetMipLevel();
 	const float LayerIndex = (float)TextureEditorPinned->GetLayer();
 	const float SliceIndex = (float)TextureEditorPinned->GetSlice();
@@ -200,7 +201,7 @@ void FTextureEditorViewportClient::Draw(FViewport* Viewport, FCanvas* Canvas)
 		// Draw a white border around the texture to show its extents
 		if (Settings.TextureBorderEnabled)
 		{
-			FCanvasBoxItem BoxItem( FVector2D(XPos, YPos), FVector2D(Width , Height ) );
+			FCanvasBoxItem BoxItem(FVector2D(XPos, YPos), FVector2D(Width + BorderSize, Height + BorderSize));
 			BoxItem.SetColor( Settings.TextureBorderColor );
 			Canvas->DrawItem( BoxItem );
 		}
@@ -308,8 +309,8 @@ bool FTextureEditorViewportClient::InputAxis(FViewport* Viewport, int32 Controll
 		{
 			TSharedPtr<STextureEditorViewport> EditorViewport = TextureEditorViewportPtr.Pin();
 
-			uint32 Width, Height, Depth, ArraySize;
-			TextureEditorPtr.Pin()->CalculateTextureDimensions(Width, Height, Depth, ArraySize);
+			int32 Width, Height, Depth, ArraySize;
+			TextureEditorPtr.Pin()->CalculateTextureDimensions(Width, Height, Depth, ArraySize, true);
 
 			if (Key == EKeys::MouseY)
 			{
@@ -382,32 +383,43 @@ void FTextureEditorViewportClient::ModifyCheckerboardTextureColors()
 FText FTextureEditorViewportClient::GetDisplayedResolution() const
 {
 	// Zero is the default size 
-	uint32 Height, Width, Depth, ArraySize;
-	TextureEditorPtr.Pin()->CalculateTextureDimensions(Width, Height, Depth, ArraySize);
+	int32 Height, Width, Depth, ArraySize;
+	TextureEditorPtr.Pin()->CalculateTextureDimensions(Width, Height, Depth, ArraySize, false);
+
+	FText CubemapInfo;
+	UTexture* Texture = TextureEditorPtr.Pin()->GetTexture();
+	if (Texture->IsA(UTextureCube::StaticClass()) || Texture->IsA(UTextureCubeArray::StaticClass()) || Texture->IsA(UTextureRenderTargetCube::StaticClass()))
+	{
+		CubemapInfo = NSLOCTEXT("TextureEditor", "DisplayedPerCubeSide", "*6 (CubeMap)");
+	}
+
+	FNumberFormattingOptions Options;
+	Options.UseGrouping = false;
+
 	if (Depth > 0)
 	{
-		return FText::Format(NSLOCTEXT("TextureEditor", "DisplayedResolutionThreeDimension", "Displayed: {0}x{1}x{2}"), FText::AsNumber(FMath::Max((uint32)1, Width)), FText::AsNumber(FMath::Max((uint32)1, Height)), FText::AsNumber(Depth));
+		return FText::Format(NSLOCTEXT("TextureEditor", "DisplayedResolutionThreeDimension", "Displayed: {0}x{1}x{2}"), FText::AsNumber(Width, &Options), FText::AsNumber(Height, &Options), FText::AsNumber(Depth, &Options));
 	}
 	else if (ArraySize > 0)
 	{
-		return FText::Format(NSLOCTEXT("TextureEditor", "DisplayedResolution", "Displayed: {0}x{1}*{2}"), FText::AsNumber(FMath::Max((uint32)1, Width)), FText::AsNumber(FMath::Max((uint32)1, Height)), FText::AsNumber(ArraySize));
+		return FText::Format(NSLOCTEXT("TextureEditor", "DisplayedResolution", "Displayed: {0}x{1}{2}*{3}"), FText::AsNumber(Width, &Options), FText::AsNumber(Height, &Options), CubemapInfo, FText::AsNumber(ArraySize, &Options));
 	}
 	else
 	{
-		return FText::Format(NSLOCTEXT("TextureEditor", "DisplayedResolutionTwoDimension", "Displayed: {0}x{1}"), FText::AsNumber(FMath::Max((uint32)1, Width)), FText::AsNumber(FMath::Max((uint32)1, Height)));
+		return FText::Format(NSLOCTEXT("TextureEditor", "DisplayedResolutionTwoDimension", "Displayed: {0}x{1}{2}"), FText::AsNumber(Width, &Options), FText::AsNumber(Height, &Options), CubemapInfo);
 	}
 }
 
 
 float FTextureEditorViewportClient::GetViewportVerticalScrollBarRatio() const
 {
-	uint32 Height = 1;
-	uint32 Width = 1;
+	int32 Height = 1;
+	int32 Width = 1;
 	float WidgetHeight = 1.0f;
 	if (TextureEditorViewportPtr.Pin()->GetVerticalScrollBar().IsValid())
 	{
-		uint32 Depth, ArraySize;
-		TextureEditorPtr.Pin()->CalculateTextureDimensions(Width, Height, Depth, ArraySize);
+		int32 Depth, ArraySize;
+		TextureEditorPtr.Pin()->CalculateTextureDimensions(Width, Height, Depth, ArraySize, true);
 
 		WidgetHeight = TextureEditorViewportPtr.Pin()->GetViewport()->GetSizeXY().Y;
 	}
@@ -418,13 +430,13 @@ float FTextureEditorViewportClient::GetViewportVerticalScrollBarRatio() const
 
 float FTextureEditorViewportClient::GetViewportHorizontalScrollBarRatio() const
 {
-	uint32 Width = 1;
-	uint32 Height = 1;
+	int32 Width = 1;
+	int32 Height = 1;
 	float WidgetWidth = 1.0f;
 	if (TextureEditorViewportPtr.Pin()->GetHorizontalScrollBar().IsValid())
 	{
-		uint32 Depth, ArraySize;
-		TextureEditorPtr.Pin()->CalculateTextureDimensions(Width, Height, Depth, ArraySize);
+		int32 Depth, ArraySize;
+		TextureEditorPtr.Pin()->CalculateTextureDimensions(Width, Height, Depth, ArraySize, true);
 
 		WidgetWidth = TextureEditorViewportPtr.Pin()->GetViewport()->GetSizeXY().X;
 	}
@@ -480,7 +492,7 @@ FVector2D FTextureEditorViewportClient::GetViewportScrollBarPositions() const
 	FVector2D Positions = FVector2D::ZeroVector;
 	if (TextureEditorViewportPtr.Pin()->GetVerticalScrollBar().IsValid() && TextureEditorViewportPtr.Pin()->GetHorizontalScrollBar().IsValid())
 	{
-		uint32 Width, Height, Depth, ArraySize;
+		int32 Width, Height, Depth, ArraySize;
 		UTexture* Texture = TextureEditorPtr.Pin()->GetTexture();
 		float VRatio = GetViewportVerticalScrollBarRatio();
 		float HRatio = GetViewportHorizontalScrollBarRatio();
@@ -489,7 +501,7 @@ FVector2D FTextureEditorViewportClient::GetViewportScrollBarPositions() const
 		float HDistFromTop = TextureEditorViewportPtr.Pin()->GetHorizontalScrollBar()->DistanceFromTop();
 		float HDistFromBottom = TextureEditorViewportPtr.Pin()->GetHorizontalScrollBar()->DistanceFromBottom();
 	
-		TextureEditorPtr.Pin()->CalculateTextureDimensions(Width, Height, Depth, ArraySize);
+		TextureEditorPtr.Pin()->CalculateTextureDimensions(Width, Height, Depth, ArraySize, true);
 
 		if ((TextureEditorViewportPtr.Pin()->GetVerticalScrollBar()->GetVisibility() == EVisibility::Visible) && VDistFromBottom < 1.0f)
 		{
diff --git a/Engine/Source/Editor/TextureEditor/Private/TextureEditorToolkit.cpp b/Engine/Source/Editor/TextureEditor/Private/TextureEditorToolkit.cpp
--- a/Engine/Source/Editor/TextureEditor/Private/TextureEditorToolkit.cpp
+++ b/Engine/Source/Editor/TextureEditor/Private/TextureEditorToolkit.cpp
@@ -237,6 +237,9 @@ void FTextureEditorToolkit::InitTextureEditor( const EToolkitMode::Type Mode, co
 
 	bIsDesaturation = false;
 
+	PreviewEffectiveTextureWidth = 0;
+	PreviewEffectiveTextureHeight = 0;
+
 	SpecifiedMipLevel = 0;
 	bUseSpecifiedMipLevel = false;
 
@@ -310,110 +313,63 @@ void FTextureEditorToolkit::InitTextureEditor( const EToolkitMode::Type Mode, co
 /* ITextureEditorToolkit interface
  *****************************************************************************/
 
-void FTextureEditorToolkit::CalculateTextureDimensions( uint32& Width, uint32& Height, uint32& Depth, uint32& ArraySize ) const
+void FTextureEditorToolkit::CalculateTextureDimensions(int32& OutWidth, int32& OutHeight, int32& OutDepth, int32& OutArraySize, bool bInIncludeBorderSize) const
 {
-	const FIntPoint LogicalSize = Texture->Source.GetLogicalSize();
-	Width = LogicalSize.X;
-	Height = LogicalSize.Y;
-	Depth = IsVolumeTexture() ? Texture->Source.GetNumLayers() : 0;
-	ArraySize = (Is2DArrayTexture() || IsCubeTexture()) ? Texture->Source.GetNumLayers() : 0;
-
-	if (!Width && !Height)
-	{
-		Width = (uint32)Texture->GetSurfaceWidth();
-		Height = (uint32)Texture->GetSurfaceHeight();
-		Depth = (uint32)Texture->GetSurfaceDepth();
-		ArraySize = Texture->GetSurfaceArraySize();
-	}
+	OutWidth = Texture->GetSurfaceWidth();
+	OutHeight = Texture->GetSurfaceHeight();
+	OutDepth = Texture->GetSurfaceDepth();
+	OutArraySize = IsArrayTexture() ? (IsCubeTexture() ? Texture->GetSurfaceArraySize() / 6 : Texture->GetSurfaceArraySize()) : 0;
+	const int32 BorderSize = GetDefault<UTextureEditorSettings>()->TextureBorderEnabled ? 1 : 0;
 
-	// catch if the Width and Height are still zero for some reason
-	if ((Width == 0) || (Height == 0))
+	if (!PreviewEffectiveTextureWidth || !PreviewEffectiveTextureHeight)
 	{
-		Width = 0;
-		Height= 0;
-		Depth = 0;
-		ArraySize = 0;
+		OutWidth = 0;
+		OutHeight = 0;
+		OutDepth = 0;
+		OutArraySize = 0;
 		return;
 	}
 
-	// See if we need to uniformly scale it to fit in viewport
-	// Cap the size to effective dimensions
-	uint32 ViewportW = TextureViewport->GetViewport()->GetSizeXY().X;
-	uint32 ViewportH = TextureViewport->GetViewport()->GetSizeXY().Y;
-	uint32 MaxWidth; 
-	uint32 MaxHeight;
-
 	// Fit is the same as fill, but doesn't scale up past 100%
 	const ETextureEditorZoomMode CurrentZoomMode = GetZoomMode();
 	if (CurrentZoomMode == ETextureEditorZoomMode::Fit || CurrentZoomMode == ETextureEditorZoomMode::Fill)
 	{
-		const UVolumeTexture* VolumeTexture = Cast<UVolumeTexture>(Texture);
-		const UTextureRenderTargetVolume* VolumeTextureRT = Cast< UTextureRenderTargetVolume>(Texture);
-
-		// Subtract off the viewport space devoted to padding (2 * PreviewPadding)
-		// so that the texture is padded on all sides
-		MaxWidth = ViewportW;
-		MaxHeight = ViewportH;
+		const int32 MaxWidth = FMath::Max(TextureViewport->GetViewport()->GetSizeXY().X - 2 * BorderSize, 0);
+		const int32 MaxHeight = FMath::Max(TextureViewport->GetViewport()->GetSizeXY().Y - 2 * BorderSize, 0);
 
-		if (IsCubeTexture())
+		if (IsVolumeTexture() && GetMutableDefault<UTextureEditorSettings>()->VolumeViewMode == ETextureEditorVolumeViewMode::TextureEditorVolumeViewMode_VolumeTrace)
 		{
-			// Cubes are displayed 2:1. 2x width if the source exists and is not an unwrapped image.
-			const bool bMultipleSourceImages = Texture->Source.GetNumSlices() > 1;
-			const bool bNoSourceImage = Texture->Source.GetNumSlices() == 0;
-			Width *= (bNoSourceImage || bMultipleSourceImages) ? 2 : 1;
-		}
-		else if (VolumeTexture || VolumeTextureRT)
-		{
-			UTextureEditorSettings& Settings = *GetMutableDefault<UTextureEditorSettings>();
-			if (Settings.VolumeViewMode == ETextureEditorVolumeViewMode::TextureEditorVolumeViewMode_VolumeTrace)
-			{
-				Width  = Height;
-			}
-			else
-			{
-				Width = FMath::CeilToInt((float)Height * (float)PreviewEffectiveTextureWidth / (float)PreviewEffectiveTextureHeight);
-			}
+			OutWidth = OutHeight = FMath::Min(MaxWidth, MaxHeight);
 		}
-
-		// First, scale up based on the size of the viewport
-		if (MaxWidth > MaxHeight)
+		else if (MaxWidth * PreviewEffectiveTextureHeight < MaxHeight * PreviewEffectiveTextureWidth)
 		{
-			Height = Height * MaxWidth / Width;
-			Width = MaxWidth;
+			OutWidth = MaxWidth;
+			OutHeight = FMath::DivideAndRoundNearest(OutWidth * PreviewEffectiveTextureHeight, PreviewEffectiveTextureWidth);
 		}
 		else
 		{
-			Width = Width * MaxHeight / Height;
-			Height = MaxHeight;
+			OutHeight = MaxHeight;
+			OutWidth = FMath::DivideAndRoundNearest(OutHeight * PreviewEffectiveTextureWidth, PreviewEffectiveTextureHeight);
 		}
 
-		// then, scale again if our width and height is impacted by the scaling
-		if (Width > MaxWidth)
-		{
-			Height = Height * MaxWidth / Width;
-			Width = MaxWidth;
-		}
-		if (Height > MaxHeight)
-		{
-			Width = Width * MaxHeight / Height;
-			Height = MaxHeight;
-		}
-		
 		// If fit, then we only want to scale down
 		// So if our natural dimensions are smaller than the viewport, we can just use those
-		if (CurrentZoomMode == ETextureEditorZoomMode::Fit)
+		if (CurrentZoomMode == ETextureEditorZoomMode::Fit && (PreviewEffectiveTextureWidth < OutWidth || PreviewEffectiveTextureHeight < OutHeight))
 		{
-			if (PreviewEffectiveTextureWidth < Width && PreviewEffectiveTextureHeight < Height)
-			{
-				Width = PreviewEffectiveTextureWidth;
-				Height = PreviewEffectiveTextureHeight;
-			}
+			OutWidth = PreviewEffectiveTextureWidth;
+			OutHeight = PreviewEffectiveTextureHeight;
 		}
 	}
 	else
 	{
-		Width = PreviewEffectiveTextureWidth * Zoom;
-		Height = PreviewEffectiveTextureHeight * Zoom;
+		OutWidth = PreviewEffectiveTextureWidth * Zoom;
+		OutHeight = PreviewEffectiveTextureHeight * Zoom;
+	}
+
+	if (bInIncludeBorderSize)
+	{
+		OutWidth += 2 * BorderSize;
+		OutHeight += 2 * BorderSize;
 	}
 }
 
@@ -726,33 +682,29 @@ void FTextureEditorToolkit::PopulateQuickInfo( )
 	const bool bIsArray = IsArrayTexture();
 	const bool bIsCube = IsCubeTexture();
 
-	const uint32 SurfaceWidth = (uint32)Texture->GetSurfaceWidth();
-	const uint32 SurfaceHeight = (uint32)Texture->GetSurfaceHeight();
-	const uint32 SurfaceDepth = (uint32)Texture->GetSurfaceDepth();
-	const uint32 NumSurfaces = (uint32)Texture->GetSurfaceArraySize();
-	const uint32 ArraySize = bIsCube ? (NumSurfaces / 6u) : NumSurfaces;
+	const int32 SurfaceWidth = Texture->GetSurfaceWidth();
+	const int32 SurfaceHeight = Texture->GetSurfaceHeight();
+	const int32 SurfaceDepth = Texture->GetSurfaceDepth();
+	const int32 NumSurfaces = Texture->GetSurfaceArraySize();
+	const int32 ArraySize = bIsArray ? (bIsCube ? NumSurfaces / 6 : NumSurfaces) : 0;
 
-	const uint32 ImportedWidth = FMath::Max<uint32>(SurfaceWidth, Texture->Source.GetSizeX());
-	const uint32 ImportedHeight =  FMath::Max<uint32>(SurfaceHeight, Texture->Source.GetSizeY());
-	const uint32 ImportedDepth = FMath::Max<uint32>(SurfaceDepth, bIsVolume ? Texture->Source.GetNumSlices() : 0);
+	const int32 ImportedWidth = FMath::Max(SurfaceWidth, Texture->Source.GetSizeX());
+	const int32 ImportedHeight =  FMath::Max(SurfaceHeight, Texture->Source.GetSizeY());
+	const int32 ImportedDepth = FMath::Max(SurfaceDepth, bIsVolume ? Texture->Source.GetNumSlices() : 0);
 
 	const FStreamableRenderResourceState SRRState = Texture->GetStreamableResourceState();
 	const int32 ActualMipBias = SRRState.IsValid() ? (SRRState.ResidentFirstLODIdx() + SRRState.AssetLODBias) : Texture->GetCachedLODBias();
-	const uint32 ActualWidth = FMath::Max<uint32>(SurfaceWidth >> ActualMipBias, 1);
-	const uint32 ActualHeight = FMath::Max<uint32>(SurfaceHeight >> ActualMipBias, 1);
-	const uint32 ActualDepth =  FMath::Max<uint32>(SurfaceDepth >> ActualMipBias, 1);
-
 	// Editor dimensions (takes user specified mip setting into account)
-	const int32 MipLevel = FMath::Max(GetMipLevel(), 0);
-	PreviewEffectiveTextureWidth = FMath::Max<uint32>(ActualWidth >> MipLevel, 1);
-	PreviewEffectiveTextureHeight = FMath::Max<uint32>(ActualHeight >> MipLevel, 1);;
-	uint32 PreviewEffectiveTextureDepth = FMath::Max<uint32>(ActualDepth >> MipLevel, 1);
+	const int32 MipLevel = ActualMipBias + FMath::Max(GetMipLevel(), 0);
+	PreviewEffectiveTextureWidth = SurfaceWidth ? FMath::Max(SurfaceWidth >> MipLevel, 1) : 0;
+	PreviewEffectiveTextureHeight = SurfaceHeight ? FMath::Max(SurfaceHeight >> MipLevel, 1) : 0;
+	const int32 PreviewEffectiveTextureDepth = SurfaceDepth ? FMath::Max(SurfaceDepth >> MipLevel, 1) : 0;
 
 	// In game max bias and dimensions
 	const int32 MaxResMipBias = Texture2D ? (Texture2D->GetNumMips() - Texture2D->GetNumMipsAllowed(true)) : Texture->GetCachedLODBias();
-	const uint32 MaxInGameWidth = FMath::Max<uint32>(SurfaceWidth >> MaxResMipBias, 1);
-	const uint32 MaxInGameHeight = FMath::Max<uint32>(SurfaceHeight >> MaxResMipBias, 1);
-	const uint32 MaxInGameDepth = FMath::Max<uint32>(SurfaceDepth >> MaxResMipBias, 1);
+	const int32 MaxInGameWidth = SurfaceWidth ? FMath::Max(SurfaceWidth >> MaxResMipBias, 1) : 0;
+	const int32 MaxInGameHeight = SurfaceHeight ? FMath::Max(SurfaceHeight >> MaxResMipBias, 1) : 0;
+	const int32 MaxInGameDepth = SurfaceDepth ? FMath::Max(SurfaceDepth >> MaxResMipBias, 1) : 0;
 
 	// Texture asset size
 	const int64 ResourceSize = PlatformDataPtr && *PlatformDataPtr ? (*PlatformDataPtr)->GetPayloadSize() : Texture->GetResourceSizeBytes(EResourceSizeMode::Exclusive);
@@ -790,8 +742,8 @@ void FTextureEditorToolkit::PopulateQuickInfo( )
 			int32 NumTilesX = 0;
 			int32 NumTilesY = 0;
 			GetBestFitForNumberOfTiles(PreviewEffectiveTextureDepth, NumTilesX, NumTilesY);
-			PreviewEffectiveTextureWidth *= (uint32)NumTilesX;
-			PreviewEffectiveTextureHeight *= (uint32)NumTilesY;
+			PreviewEffectiveTextureWidth *= NumTilesX;
+			PreviewEffectiveTextureHeight *= NumTilesY;
 		}
 	}
 	else if (bIsArray)
@@ -857,8 +809,8 @@ double FTextureEditorToolkit::CalculateDisplayedZoomLevel() const
 		return Zoom;
 	}
 
-	uint32 DisplayWidth, DisplayHeight, DisplayDepth, DisplayArraySize;
-	CalculateTextureDimensions(DisplayWidth, DisplayHeight, DisplayDepth, DisplayArraySize);
+	int32 DisplayWidth, DisplayHeight, DisplayDepth, DisplayArraySize;
+	CalculateTextureDimensions(DisplayWidth, DisplayHeight, DisplayDepth, DisplayArraySize, false);
 	if (PreviewEffectiveTextureHeight != 0)
 	{
 		return (double)DisplayHeight / PreviewEffectiveTextureHeight;
diff --git a/Engine/Source/Editor/TextureEditor/Private/TextureEditorToolkit.h b/Engine/Source/Editor/TextureEditor/Private/TextureEditorToolkit.h
--- a/Engine/Source/Editor/TextureEditor/Private/TextureEditorToolkit.h
+++ b/Engine/Source/Editor/TextureEditor/Private/TextureEditorToolkit.h
@@ -61,7 +61,7 @@ public:
 	virtual void UnregisterTabSpawners( const TSharedRef<class FTabManager>& TabManager ) override;
 
 	// ITextureEditorToolkit interface
-	virtual void CalculateTextureDimensions( uint32& Width, uint32& Height, uint32& Depth, uint32& ArraySize ) const override;
+	virtual void CalculateTextureDimensions(int32& OutWidth, int32& OutHeight, int32& OutDepth, int32& OutArraySize, bool bInIncludeBorderSize) const override;
 	virtual ESimpleElementBlendMode GetColourChannelBlendMode( ) const override;
 	virtual int32 GetMipLevel( ) const override;
 	virtual int32 GetLayer() const override;
@@ -408,8 +408,8 @@ private:
 	bool bIsDesaturation;
 
 	/** The maximum width/height at which the texture will render in the preview window */
-	uint32 PreviewEffectiveTextureWidth;
-	uint32 PreviewEffectiveTextureHeight;
+	int32 PreviewEffectiveTextureWidth;
+	int32 PreviewEffectiveTextureHeight;
 
 	/** Which mip level should be shown */
 	int32 SpecifiedMipLevel;
diff --git a/Engine/Source/Editor/TextureEditor/Public/Interfaces/ITextureEditorToolkit.h b/Engine/Source/Editor/TextureEditor/Public/Interfaces/ITextureEditorToolkit.h
--- a/Engine/Source/Editor/TextureEditor/Public/Interfaces/ITextureEditorToolkit.h
+++ b/Engine/Source/Editor/TextureEditor/Public/Interfaces/ITextureEditorToolkit.h
@@ -26,8 +26,8 @@ public:
 	/** Refreshes the quick info panel */
 	virtual void PopulateQuickInfo() = 0;
 
-	/** Calculates the display size of the texture */
-	virtual void CalculateTextureDimensions(uint32& Width, uint32& Height, uint32& Depth, uint32& ArraySize) const = 0;
+	/** Calculates the display size of the texture, optionally including the size of the displayed texture border */
+	virtual void CalculateTextureDimensions(int32& OutWidth, int32& OutHeight, int32& OutDepth, int32& OutArraySize, bool bInIncludeBorderSize) const = 0;
 
 	/** Accessors */ 
 	virtual int32 GetMipLevel() const = 0;
