Enable mipmap preview for cubemaps.

diff --git a/Engine/Source/Editor/TextureEditor/Private/TextureEditorToolkit.cpp b/Engine/Source/Editor/TextureEditor/Private/TextureEditorToolkit.cpp
--- a/Engine/Source/Editor/TextureEditor/Private/TextureEditorToolkit.cpp
+++ b/Engine/Source/Editor/TextureEditor/Private/TextureEditorToolkit.cpp
@@ -2088,14 +2088,7 @@ ECheckBoxState FTextureEditorToolkit::HandleMipLevelCheckBoxIsChecked( ) const
 
 bool FTextureEditorToolkit::HandleMipLevelCheckBoxIsEnabled( ) const
 {
-	UTextureCube* TextureCube = Cast<UTextureCube>(Texture);
-
-	if (GetMaxMipLevel().Get(MIPLEVEL_MAX) <= 0 || TextureCube)
-	{
-		return false;
-	}
-
-	return true;
+	return GetMaxMipLevel().Get(MIPLEVEL_MAX) > 0;
 }
 
 void FTextureEditorToolkit::HandleMipLevelChanged(int32 NewMipLevel)
