Add support for previewing individual faces of TextureCubeArray elements in the texture editor.

diff --git a/Engine/Shaders/Private/SimpleElementPixelShader.usf b/Engine/Shaders/Private/SimpleElementPixelShader.usf
--- a/Engine/Shaders/Private/SimpleElementPixelShader.usf
+++ b/Engine/Shaders/Private/SimpleElementPixelShader.usf
@@ -350,7 +350,7 @@ float3 GetCubemapVector(float2 ScaledUVs, uint CubeFace)
 	return CubeCoordinates;
 }
 
-// .x:MipLevel, .y: SliceIndex, .zw:unused
+// .x:MipLevel, .y: SliceIndex, .z: FaceIndex, .w:unused
 float4 PackedProperties0;
 float4x4 ColorWeights;
 
@@ -369,42 +369,41 @@ void CubemapTextureProperties(
 	)
 {
 	float MipLevel = PackedProperties0.x;
-	float SliceIndex = PackedProperties0.y;
+	float FaceIndex = PackedProperties0.z;
 	float4 BaseColor;
 
 #if TEXTURECUBE_ARRAY
 
-	// for a TextureCube array SliceIndex represents an index of a cubemap in the array
-	// each slice of a TextureCube array is displayed as a longlat unwrap
-
-	float2 Angles = float2(2 * PI * (TextureCoordinate.x + 0.5f), PI * TextureCoordinate.y);
-	float s = sin(Angles.y);
-	float3 Direction = float3(s * sin(Angles.x), cos(Angles.y), -s * cos(Angles.x));
-
-	float CurrentSliceIndex = SliceIndex;
-	if (CurrentSliceIndex < 0.0)
+	// SliceIndex represents an index of a cubemap in the array
+	float SliceIndex = PackedProperties0.y;
+	if (SliceIndex < 0.0)
 	{
 		// if the slice index is not specified, display a portion of each cubemap from the TextureCube array
-		CurrentSliceIndex = floor(TextureCoordinate.y * NumSlices);
+		SliceIndex = floor(TextureCoordinate.y * NumSlices);
 	}
 
-#if FEATURE_LEVEL >= FEATURE_LEVEL_SM5
-	BaseColor = TextureCubeArraySampleLevel(CubeTexture, CubeTextureSampler, Direction.xzy, CurrentSliceIndex, MipLevel);
-#else
-	BaseColor = 0;
-#endif
+	if (FaceIndex >= 0.0)
+	{
+		BaseColor = TextureCubeArraySampleLevel(CubeTexture, CubeTextureSampler, GetCubemapVector(TextureCoordinate * 2 - 1, FaceIndex), SliceIndex, MipLevel);
+	}
+	else
+	{
+		// if the face index is not specified, display the cubemap as a longlat unwrap
+		float2 Angles = float2(2 * PI * (TextureCoordinate.x + 0.5f), PI * TextureCoordinate.y);
+		float s = sin(Angles.y);
+		float3 Direction = float3(s * sin(Angles.x), cos(Angles.y), -s * cos(Angles.x));
+		BaseColor = TextureCubeArraySampleLevel(CubeTexture, CubeTextureSampler, Direction.xzy, SliceIndex, MipLevel);
+	}
 
 #else //TEXTURECUBE_ARRAY
 
-	// for a TextureCube SliceIndex represents an index of a cubemap face
-
-	if (SliceIndex >= 0.0)
+	if (FaceIndex >= 0.0)
 	{
-		BaseColor = TextureCubeSampleLevel(CubeTexture, CubeTextureSampler, GetCubemapVector(TextureCoordinate * 2 - 1, SliceIndex), MipLevel);
+		BaseColor = TextureCubeSampleLevel(CubeTexture, CubeTextureSampler, GetCubemapVector(TextureCoordinate * 2 - 1, FaceIndex), MipLevel);
 	}
 	else
 	{
-		// if the slice index is not specified, display the cubemap as a longlat unwrap
+		// if the face index is not specified, display the cubemap as a longlat unwrap
 		float2 Angles = float2(2 * PI * (TextureCoordinate.x + 0.5f), PI * TextureCoordinate.y);
 		float s = sin(Angles.y);
 		float3 Direction = float3(s * sin(Angles.x), cos(Angles.y), -s * cos(Angles.x));
diff --git a/Engine/Source/Editor/TextureEditor/Private/Models/TextureEditorViewportClient.cpp b/Engine/Source/Editor/TextureEditor/Private/Models/TextureEditorViewportClient.cpp
--- a/Engine/Source/Editor/TextureEditor/Private/Models/TextureEditorViewportClient.cpp
+++ b/Engine/Source/Editor/TextureEditor/Private/Models/TextureEditorViewportClient.cpp
@@ -112,7 +112,8 @@ void FTextureEditorViewportClient::Draw(FViewport* Viewport, FCanvas* Canvas)
 	{
 		if (TextureCube || TextureCubeArray || RTTextureCube)
 		{
-			BatchedElementParameters = new FMipLevelBatchedElementParameters(MipLevel, SliceIndex, TextureCubeArray != nullptr, false);
+			const float FaceIndex = (float)TextureEditorPinned->GetFace();
+			BatchedElementParameters = new FMipLevelBatchedElementParameters(MipLevel, SliceIndex, FaceIndex, TextureCubeArray != nullptr, false);
 		}
 		else if (VolumeTexture)
 		{
diff --git a/Engine/Source/Editor/TextureEditor/Private/TextureEditorToolkit.cpp b/Engine/Source/Editor/TextureEditor/Private/TextureEditorToolkit.cpp
--- a/Engine/Source/Editor/TextureEditor/Private/TextureEditorToolkit.cpp
+++ b/Engine/Source/Editor/TextureEditor/Private/TextureEditorToolkit.cpp
@@ -248,6 +248,9 @@ void FTextureEditorToolkit::InitTextureEditor( const EToolkitMode::Type Mode, co
 	SpecifiedSlice = 0;
 	bUseSpecifiedSlice = false;
 
+	SpecifiedFace = 0;
+	bUseSpecifiedFace = false;
+
 	SavedCompressionSetting = false;
 
 	// Start at whatever the last used zoom mode was
@@ -421,6 +424,11 @@ int32 FTextureEditorToolkit::GetSlice() const
 	return GetUseSpecifiedSlice() ? HandleSliceEntryBoxValue().Get(-1) : -1;
 }
 
+int32 FTextureEditorToolkit::GetFace() const
+{
+	return GetUseSpecifiedFace() ? SpecifiedFace : -1;
+}
+
 UTexture* FTextureEditorToolkit::GetTexture( ) const
 {
 	return Texture;
@@ -455,6 +463,11 @@ bool FTextureEditorToolkit::GetUseSpecifiedSlice() const
 	return HandleSliceCheckBoxIsEnabled() && bUseSpecifiedSlice;
 }
 
+bool FTextureEditorToolkit::GetUseSpecifiedFace() const
+{
+	return bUseSpecifiedFace;
+}
+
 double FTextureEditorToolkit::GetCustomZoomLevel( ) const
 {
 	return Zoom;
@@ -713,10 +726,9 @@ void FTextureEditorToolkit::PopulateQuickInfo( )
 	SizeOptions.UseGrouping = false;
 	SizeOptions.MaximumFractionalDigits = 0;
 
-	// Cubes are previewed as unwrapped 2D textures in case when slice index is not selected.
+	// Cubes are previewed as longlat unwrapped 2D textures in case when face index is not specified.
 	// These have 2x the width of a cube face.
-	// Cube arrays are always displayed as longlat unwraps.
-	if (IsCubeTexture() && (IsArrayTexture() || GetSlice() < 0))
+	if (IsCubeTexture() && GetFace() < 0)
 	{
 		PreviewEffectiveTextureWidth *= 2;
 	}
@@ -1661,6 +1673,7 @@ void FTextureEditorToolkit::FillToolbar(FToolBarBuilder& ToolbarBuilder)
 	TSharedRef<SWidget> LODControl = MakeLODControlWidget();
 	TSharedRef<SWidget> LayerControl = MakeLayerControlWidget();
 	TSharedRef<SWidget> SliceControl = MakeSliceControlWidget();
+	TSharedRef<SWidget> FaceControl = MakeFaceControlWidget();
 	TSharedRef<SWidget> ExposureControl = MakeExposureContolWidget();
 	TSharedPtr<SWidget> OptionalOpacityControl = IsVolumeTexture() ? TSharedPtr<SWidget>(MakeOpacityControlWidget()) : nullptr;
 	TSharedRef<SWidget> ZoomControl = MakeZoomControlWidget();
@@ -1706,6 +1719,15 @@ void FTextureEditorToolkit::FillToolbar(FToolBarBuilder& ToolbarBuilder)
 			ToolbarBuilder.EndSection();
 		}
 
+		if (IsCubeTexture())
+		{
+			ToolbarBuilder.BeginSection("Faces");
+			{
+				ToolbarBuilder.AddWidget(FaceControl);
+			}
+			ToolbarBuilder.EndSection();
+		}
+
 		if (OptionalOpacityControl.IsValid())
 		{
 			ToolbarBuilder.BeginSection("Opacity");
@@ -2137,9 +2159,8 @@ bool FTextureEditorToolkit::HasLayers() const
 
 bool FTextureEditorToolkit::HasSlices() const
 {
-	// currently supports Texture2D arrays, cubemaps and cubemap arrays
 	// slice selection should be supported even for a texture array with less than two elements, because array elements can be added dynamically
-	return Is2DArrayTexture() || IsCubeTexture();
+	return IsArrayTexture();
 }
 
 int32 FTextureEditorToolkit::GetNumSlices() const
@@ -2157,10 +2178,6 @@ int32 FTextureEditorToolkit::GetNumSlices() const
 		// for a TextureCube array SliceIndex represents an index of a cubemap in the array
 		return Cast<UTextureCubeArray>(Texture)->GetSurfaceArraySize() / 6;
 	}
-	else if (IsCubeTexture() && !IsArrayTexture())
-	{
-		return 6;
-	}
 	return 1;
 }
 
@@ -2190,6 +2207,27 @@ void FTextureEditorToolkit::HandleSliceEntryBoxChanged(int32 NewSlice)
 	PopulateQuickInfo();
 }
 
+ECheckBoxState FTextureEditorToolkit::HandleFaceCheckBoxIsChecked() const
+{
+	return GetUseSpecifiedFace() ? ECheckBoxState::Checked : ECheckBoxState::Unchecked;
+}
+
+void FTextureEditorToolkit::HandleFaceCheckBoxCheckedStateChanged(ECheckBoxState InNewState)
+{
+	bUseSpecifiedFace = InNewState == ECheckBoxState::Checked;
+}
+
+TOptional<int32> FTextureEditorToolkit::HandleFaceEntryBoxValue() const
+{
+	return FMath::Clamp<int32>(SpecifiedFace, 0, 5);
+}
+
+void FTextureEditorToolkit::HandleFaceEntryBoxChanged(int32 NewFace)
+{
+	SpecifiedFace = FMath::Clamp<int32>(NewFace, 0, 5);
+	PopulateQuickInfo();
+}
+
 bool FTextureEditorToolkit::HandleReimportActionCanExecute( ) const
 {
 	if (Texture->IsA<ULightMapTexture2D>() || Texture->IsA<UShadowMapTexture2D>() || Texture->IsA<UTexture2DDynamic>() || Texture->IsA<UTextureRenderTarget>() || Texture->IsA<UCurveLinearColorAtlas>())
@@ -2731,6 +2769,46 @@ TSharedRef<SWidget> FTextureEditorToolkit::MakeSliceControlWidget()
 	return SliceControl;
 }
 
+TSharedRef<SWidget> FTextureEditorToolkit::MakeFaceControlWidget()
+{
+	TSharedRef<SWidget> FaceControl = SNew(SBox)
+		.WidthOverride(212.0f)
+		[
+			SNew(SHorizontalBox)
+			+ SHorizontalBox::Slot()
+			.Padding(4.0f, 0.0f, 4.0f, 0.0f)
+			.VAlign(VAlign_Center)
+			.AutoWidth()
+			[
+				SNew(STextBlock)
+				.Text(NSLOCTEXT("TextureEditor", "Face", "Face"))
+			]
+			+ SHorizontalBox::Slot()
+			.Padding(4.0f, 0.0f, 2.0f, 0.0f)
+			.VAlign(VAlign_Center)
+			.AutoWidth()
+			[
+				SNew(SCheckBox)
+				.IsChecked(this, &FTextureEditorToolkit::HandleFaceCheckBoxIsChecked)
+				.OnCheckStateChanged(this, &FTextureEditorToolkit::HandleFaceCheckBoxCheckedStateChanged)
+			]
+			+ SHorizontalBox::Slot()
+			.Padding(0.0f, 0.0f, 4.0f, 0.0f)
+			.VAlign(VAlign_Center)
+			[
+				SNew(SNumericEntryBox<int32>)
+				.IsEnabled(this, &FTextureEditorToolkit::GetUseSpecifiedFace)
+				.AllowSpin(true)
+				.MinSliderValue(0)
+				.MaxSliderValue(5)
+				.Value(this, &FTextureEditorToolkit::HandleFaceEntryBoxValue)
+				.OnValueChanged(this, &FTextureEditorToolkit::HandleFaceEntryBoxChanged)
+			]
+		];
+
+	return FaceControl;
+}
+
 TSharedRef<SWidget> FTextureEditorToolkit::MakeExposureContolWidget()
 {
 	TSharedRef<SWidget> ExposureControl = SNew(SBox)
diff --git a/Engine/Source/Editor/TextureEditor/Private/TextureEditorToolkit.h b/Engine/Source/Editor/TextureEditor/Private/TextureEditorToolkit.h
--- a/Engine/Source/Editor/TextureEditor/Private/TextureEditorToolkit.h
+++ b/Engine/Source/Editor/TextureEditor/Private/TextureEditorToolkit.h
@@ -66,10 +66,12 @@ public:
 	virtual int32 GetMipLevel( ) const override;
 	virtual int32 GetLayer() const override;
 	virtual int32 GetSlice() const override;
+	virtual int32 GetFace() const override;
 	virtual UTexture* GetTexture( ) const override;
 	virtual bool HasValidTextureResource( ) const override;
 	virtual bool GetUseSpecifiedMip( ) const override;
 	virtual bool GetUseSpecifiedSlice() const override;
+	virtual bool GetUseSpecifiedFace() const override;
 	virtual double GetCustomZoomLevel( ) const override;
 	virtual void SetCustomZoomLevel(double ZoomValue) override;
 	virtual void PopulateQuickInfo( ) override;
@@ -261,6 +263,18 @@ private:
 	// Callback for changing the value of the slice index entry box.
 	void HandleSliceEntryBoxChanged(int32 Slice);
 
+	// Callback for getting the checked state of the face check box.
+	ECheckBoxState HandleFaceCheckBoxIsChecked() const;
+
+	// Callback for changing the checked state of the face check box.
+	void HandleFaceCheckBoxCheckedStateChanged(ECheckBoxState InNewState);
+
+	// Callback for getting the value of the face index entry box.
+	TOptional<int32> HandleFaceEntryBoxValue() const;
+
+	// Callback for changing the value of the face index entry box.
+	void HandleFaceEntryBoxChanged(int32 Face);
+
 	// Callback for determining whether the Reimport action can execute.
 	bool HandleReimportActionCanExecute( ) const;
 
@@ -342,6 +356,7 @@ private:
 	TSharedRef<SWidget> MakeLODControlWidget();
 	TSharedRef<SWidget> MakeLayerControlWidget();
 	TSharedRef<SWidget> MakeSliceControlWidget();
+	TSharedRef<SWidget> MakeFaceControlWidget();
 	TSharedRef<SWidget> MakeExposureContolWidget();
 	TSharedRef<SWidget> MakeOpacityControlWidget();
 	TSharedRef<SWidget> MakeZoomControlWidget();
@@ -426,6 +441,12 @@ private:
 	/** When true, the specified slice index value is used, otherwise the value of -1 is used.*/
 	bool bUseSpecifiedSlice;
 
+	/** Which cubemap face should be shown */
+	int32 SpecifiedFace;
+
+	/** When true, the specified cubemap face index value is used, otherwise the value of -1 is used.*/
+	bool bUseSpecifiedFace;
+
 	/** During re-import, cache this setting so it can be restored if necessary */
 	bool SavedCompressionSetting;
 
diff --git a/Engine/Source/Editor/TextureEditor/Public/Interfaces/ITextureEditorToolkit.h b/Engine/Source/Editor/TextureEditor/Public/Interfaces/ITextureEditorToolkit.h
--- a/Engine/Source/Editor/TextureEditor/Public/Interfaces/ITextureEditorToolkit.h
+++ b/Engine/Source/Editor/TextureEditor/Public/Interfaces/ITextureEditorToolkit.h
@@ -33,9 +33,11 @@ public:
 	virtual int32 GetMipLevel() const = 0;
 	virtual int32 GetLayer() const = 0;
 	virtual int32 GetSlice() const = 0;
+	virtual int32 GetFace() const = 0;
 	virtual ESimpleElementBlendMode GetColourChannelBlendMode() const = 0;
 	virtual bool GetUseSpecifiedMip() const = 0;
 	virtual bool GetUseSpecifiedSlice() const = 0;
+	virtual bool GetUseSpecifiedFace() const = 0;
 	virtual double GetCustomZoomLevel() const = 0;
 	virtual void SetCustomZoomLevel( double ZoomValue ) = 0;
 	virtual void ZoomIn() = 0;
diff --git a/Engine/Source/Editor/UnrealEd/Private/ThumbnailRendering/TextureThumbnailRenderer.cpp b/Engine/Source/Editor/UnrealEd/Private/ThumbnailRendering/TextureThumbnailRenderer.cpp
--- a/Engine/Source/Editor/UnrealEd/Private/ThumbnailRendering/TextureThumbnailRenderer.cpp
+++ b/Engine/Source/Editor/UnrealEd/Private/ThumbnailRendering/TextureThumbnailRenderer.cpp
@@ -72,7 +72,7 @@ void UTextureThumbnailRenderer::Draw(UObject* Object, int32 X, int32 Y, uint32 W
 		if(TextureCube || TextureCubeArray || RTTextureCube)
 		{
 			// is released by the render thread when it was rendered
-			BatchedElementParameters = new FMipLevelBatchedElementParameters((float)0, (float)-1, TextureCubeArray != nullptr);
+			BatchedElementParameters = new FMipLevelBatchedElementParameters((float)0, (float)-1, (float)-1, TextureCubeArray != nullptr);
 			
 			// If the thumbnail is square then make it 2:1 for cubes.
 			if(Width == Height)
diff --git a/Engine/Source/Runtime/Engine/Private/CubemapUnwrapUtils.cpp b/Engine/Source/Runtime/Engine/Private/CubemapUnwrapUtils.cpp
--- a/Engine/Source/Runtime/Engine/Private/CubemapUnwrapUtils.cpp
+++ b/Engine/Source/Runtime/Engine/Private/CubemapUnwrapUtils.cpp
@@ -36,7 +36,7 @@ namespace CubemapHelpers
 	bool GenerateLongLatUnwrap(const FTextureResource* TextureResource, const uint32 AxisDimenion, const EPixelFormat SourcePixelFormat, TArray64<uint8>& BitsOUT, FIntPoint& SizeOUT, EPixelFormat& FormatOUT)
 	{
 		TRefCountPtr<FBatchedElementParameters> BatchedElementParameters;
-		BatchedElementParameters = new FMipLevelBatchedElementParameters((float)0, (float)-1, false, true);
+		BatchedElementParameters = new FMipLevelBatchedElementParameters((float)0, (float)-1, (float)-1, false, true);
 		const FIntPoint LongLatDimensions(AxisDimenion * 2, AxisDimenion);
 
 		// If the source format is 8 bit per channel or less then select a LDR target format.
@@ -145,13 +145,13 @@ bool FCubemapTexturePropertiesPS::ShouldCompilePermutation(const FGlobalShaderPe
 	return true;
 }
 
-void FCubemapTexturePropertiesPS::SetParameters(FRHICommandList& RHICmdList, const FTexture* Texture, const FMatrix& ColorWeightsValue, float MipLevel, float SliceIndex, float GammaValue, bool bIsTextureCubeArray)
+void FCubemapTexturePropertiesPS::SetParameters(FRHICommandList& RHICmdList, const FTexture* Texture, const FMatrix& ColorWeightsValue, float MipLevel, float SliceIndex, float FaceIndex, float GammaValue, bool bIsTextureCubeArray)
 {
 	FRHIPixelShader* ShaderRHI = RHICmdList.GetBoundPixelShader();
 
 	SetTextureParameter(RHICmdList, ShaderRHI, CubeTexture, CubeTextureSampler, Texture);
 
-	FVector4f PackedProperties0Value(MipLevel, SliceIndex, 0, 0);
+	FVector4f PackedProperties0Value(MipLevel, SliceIndex, FaceIndex, 0);
 	SetShaderValue(RHICmdList, ShaderRHI, PackedProperties0, PackedProperties0Value);
 	SetShaderValue(RHICmdList, ShaderRHI, ColorWeights, (FMatrix44f)ColorWeightsValue);
 	SetShaderValue(RHICmdList, ShaderRHI, Gamma, GammaValue);
@@ -187,7 +187,7 @@ void FMipLevelBatchedElementParameters::BindShaders(FRHICommandList& RHICmdList,
 	SetGraphicsPipelineState(RHICmdList, GraphicsPSOInit, 0);
 
 	VertexShader->SetParameters(RHICmdList, InTransform);
-	PixelShader->SetParameters(RHICmdList, Texture, ColorWeights, MipLevel, SliceIndex, InGamma, bIsTextureCubeArray);
+	PixelShader->SetParameters(RHICmdList, Texture, ColorWeights, MipLevel, SliceIndex, FaceIndex, InGamma, bIsTextureCubeArray);
 }
 
 void FIESLightProfilePS::SetParameters( FRHICommandList& RHICmdList, const FTexture* Texture, float InBrightnessInLumens )
diff --git a/Engine/Source/Runtime/Engine/Public/CubemapUnwrapUtils.h b/Engine/Source/Runtime/Engine/Public/CubemapUnwrapUtils.h
--- a/Engine/Source/Runtime/Engine/Public/CubemapUnwrapUtils.h
+++ b/Engine/Source/Runtime/Engine/Public/CubemapUnwrapUtils.h
@@ -79,7 +79,7 @@ public:
 	FCubemapTexturePropertiesPS(const ShaderMetaType::CompiledShaderInitializerType& Initializer);
 
 	static bool ShouldCompilePermutation(const FGlobalShaderPermutationParameters& Parameters);
-	void SetParameters(FRHICommandList& RHICmdList, const FTexture* Texture, const FMatrix& ColorWeightsValue, float MipLevel, float SliceIndex, float GammaValue, bool bIsTextureCubeArray);
+	void SetParameters(FRHICommandList& RHICmdList, const FTexture* Texture, const FMatrix& ColorWeightsValue, float MipLevel, float SliceIndex, float FaceIndex, float GammaValue, bool bIsTextureCubeArray);
 
 private:
 	LAYOUT_FIELD(FShaderResourceParameter, CubeTexture);
@@ -94,10 +94,11 @@ private:
 class ENGINE_API FMipLevelBatchedElementParameters : public FBatchedElementParameters
 {
 public:
-	FMipLevelBatchedElementParameters(float InMipLevel, float InSliceIndex, bool bInIsTextureCubeArray, bool bInHDROutput = false)
+	FMipLevelBatchedElementParameters(float InMipLevel, float InSliceIndex, float InFaceIndex, bool bInIsTextureCubeArray, bool bInHDROutput = false)
 		: bHDROutput(bInHDROutput)
 		, MipLevel(InMipLevel)
 		, SliceIndex(InSliceIndex)
+		, FaceIndex(InFaceIndex)
 		, bIsTextureCubeArray(bInIsTextureCubeArray)
 	{
 	}
@@ -111,6 +112,7 @@ private:
 	/** Parameters that need to be passed to the shader */
 	float MipLevel;
 	float SliceIndex;
+	float FaceIndex;
 
 	/** Parameters that are used to select a shader permutation */
 	bool bIsTextureCubeArray;
