Implement 3D preview mode for cubemaps in the texture editor.

diff --git a/Engine/Shaders/Private/SimpleElementPixelShader.usf b/Engine/Shaders/Private/SimpleElementPixelShader.usf
--- a/Engine/Shaders/Private/SimpleElementPixelShader.usf
+++ b/Engine/Shaders/Private/SimpleElementPixelShader.usf
@@ -317,45 +317,14 @@ TextureCube CubeTexture;
 
 SamplerState CubeTextureSampler;
 
-float3 GetCubemapVector(float2 ScaledUVs, uint CubeFace)
-{
-	float3 CubeCoordinates;
-	
-	// calculation of the cube coordinates matches the DDS orientaion and face order: positive x, negative x, positive y, negative y, positive z, negative z
-	if (CubeFace == 0)
-	{
-		CubeCoordinates = float3(1, -ScaledUVs.y, -ScaledUVs.x);
-	}
-	else if (CubeFace == 1)
-	{
-		CubeCoordinates = float3(-1, -ScaledUVs.y, ScaledUVs.x);
-	}
-	else if (CubeFace == 2)
-	{
-		CubeCoordinates = float3(ScaledUVs.x, 1, ScaledUVs.y);
-	}
-	else if (CubeFace == 3)
-	{
-		CubeCoordinates = float3(ScaledUVs.x, -1, -ScaledUVs.y);
-	}
-	else if (CubeFace == 4)
-	{
-		CubeCoordinates = float3(ScaledUVs.x, -ScaledUVs.y, 1);
-	}
-	else
-	{
-		CubeCoordinates = float3(-ScaledUVs.x, -ScaledUVs.y, -1);
-	}
-
-	return CubeCoordinates;
-}
-
-// .x:MipLevel, .y: SliceIndex, .z: FaceIndex, .w:unused
+// .x:MipLevel, .y:bShowLongLatUnwrap, .zw:unused
 float4 PackedProperties0;
 float4x4 ColorWeights;
+float4x4 ViewMatrix;
 
 #if TEXTURECUBE_ARRAY
 float NumSlices;
+float SliceIndex;
 #endif
 
 // used in the Texture info dialog when looking at a cubemap texture
@@ -369,79 +338,49 @@ void CubemapTextureProperties(
 	)
 {
 	float MipLevel = PackedProperties0.x;
-	float FaceIndex = PackedProperties0.z;
+	float bShowLongLatUnwrap = PackedProperties0.y;
+	float3 CubeCoordinates;
 	float4 BaseColor;
 
-#if TEXTURECUBE_ARRAY
+	if (bShowLongLatUnwrap > 0.0)
+	{
+		float2 Angles = float2(2 * PI * (TextureCoordinate.x + 0.5f), PI * TextureCoordinate.y);
+		float s = sin(Angles.y);
+		CubeCoordinates = float3(s * sin(Angles.x), -s * cos(Angles.x), cos(Angles.y));
+	}
+	else
+	{
+		float2 ScaledUVs = TextureCoordinate * 2 - 1;
+		// identity view matrix displays the face 0 properly oriented in space
+		CubeCoordinates = mul(float3(1, ScaledUVs.x, -ScaledUVs.y), ViewMatrix);
+	}
 
+#if TEXTURECUBE_ARRAY
 	// SliceIndex represents an index of a cubemap in the array
-	float SliceIndex = PackedProperties0.y;
 	if (SliceIndex < 0.0)
 	{
 		// if the slice index is not specified, display a portion of each cubemap from the TextureCube array
 		SliceIndex = floor(TextureCoordinate.y * NumSlices);
 	}
+#endif
 
-	if (FaceIndex >= 0.0)
-	{
-		float4 CubeArrayCoordinates = float4(GetCubemapVector(TextureCoordinate * 2 - 1, FaceIndex), SliceIndex);
-		if (MipLevel >= 0.0)
-		{
-			BaseColor = TextureCubeArraySampleLevel(CubeTexture, CubeTextureSampler, CubeArrayCoordinates, MipLevel);
-		}
-		else
-		{
-			BaseColor = TextureCubeArraySample(CubeTexture, CubeTextureSampler, CubeArrayCoordinates);
-		}
-	}
-	else
-	{
-		// if the face index is not specified, display the cubemap as a longlat unwrap
-		float2 Angles = float2(2 * PI * (TextureCoordinate.x + 0.5f), PI * TextureCoordinate.y);
-		float s = sin(Angles.y);
-		float4 CubeArrayCoordinates = float4(s * sin(Angles.x), -s * cos(Angles.x), cos(Angles.y), SliceIndex);
-		if (MipLevel >= 0.0)
-		{
-			BaseColor = TextureCubeArraySampleLevel(CubeTexture, CubeTextureSampler, CubeArrayCoordinates, MipLevel);
-		}
-		else
-		{
-			BaseColor = TextureCubeArraySample(CubeTexture, CubeTextureSampler, CubeArrayCoordinates);
-		}
-	}
-
-#else //TEXTURECUBE_ARRAY
-
-	if (FaceIndex >= 0.0)
+	if (MipLevel >= 0.0)
 	{
-		float3 CubeCoordinates = GetCubemapVector(TextureCoordinate * 2 - 1, FaceIndex);
-		if (MipLevel >= 0.0)
-		{
-			BaseColor = TextureCubeSampleLevel(CubeTexture, CubeTextureSampler, CubeCoordinates, MipLevel);
-		}
-		else
-		{
-			BaseColor = TextureCubeSample(CubeTexture, CubeTextureSampler, CubeCoordinates);
-		}
+#if TEXTURECUBE_ARRAY
+		BaseColor = TextureCubeArraySampleLevel(CubeTexture, CubeTextureSampler, float4(CubeCoordinates, SliceIndex), MipLevel);
+#else
+		BaseColor = TextureCubeSampleLevel(CubeTexture, CubeTextureSampler, CubeCoordinates, MipLevel);
+#endif
 	}
 	else
 	{
-		// if the face index is not specified, display the cubemap as a longlat unwrap
-		float2 Angles = float2(2 * PI * (TextureCoordinate.x + 0.5f), PI * TextureCoordinate.y);
-		float s = sin(Angles.y);
-		float3 CubeCoordinates = float3(s * sin(Angles.x), -s * cos(Angles.x), cos(Angles.y));
-		if (MipLevel >= 0.0)
-		{
-			BaseColor = TextureCubeSampleLevel(CubeTexture, CubeTextureSampler, CubeCoordinates, MipLevel);
-		}
-		else
-		{
-			BaseColor = TextureCubeSample(CubeTexture, CubeTextureSampler, CubeCoordinates);
-		}
+#if TEXTURECUBE_ARRAY
+		BaseColor = TextureCubeArraySample(CubeTexture, CubeTextureSampler, float4(CubeCoordinates, SliceIndex));
+#else
+		BaseColor = TextureCubeSample(CubeTexture, CubeTextureSampler, CubeCoordinates);
+#endif
 	}
 
-#endif //TEXTURECUBE_ARRAY
-
 	float4 FinalColor;
 
 	// Seperate the Color weights and use against the Base colour to detrmine the actual colour from our filter
diff --git a/Engine/Source/Editor/TextureEditor/Classes/TextureEditorSettings.h b/Engine/Source/Editor/TextureEditor/Classes/TextureEditorSettings.h
--- a/Engine/Source/Editor/TextureEditor/Classes/TextureEditorSettings.h
+++ b/Engine/Source/Editor/TextureEditor/Classes/TextureEditorSettings.h
@@ -27,6 +27,13 @@ enum ETextureEditorVolumeViewMode
 	TextureEditorVolumeViewMode_VolumeTrace UMETA(DisplayName="Trace Into Volume"),
 };
 
+UENUM()
+enum ETextureEditorCubemapViewMode
+{
+	TextureEditorCubemapViewMode_2DView UMETA(DisplayName = "2D View"),
+	TextureEditorCubemapViewMode_3DView UMETA(DisplayName = "3D View"),
+};
+
 UENUM()
 enum class ETextureEditorZoomMode : uint8
 {
@@ -50,10 +57,14 @@ public:
 	UPROPERTY(config)
 	TEnumAsByte<ETextureEditorBackgrounds> Background;
 
-	/** The type of display when viewing volume textures. */
+	/** The view mode when previewing volume textures. */
 	UPROPERTY(config)
 	TEnumAsByte<ETextureEditorVolumeViewMode> VolumeViewMode;
 
+	/** The view mode when previewing cubemap textures. */
+	UPROPERTY(config)
+	TEnumAsByte<ETextureEditorCubemapViewMode> CubemapViewMode;
+
 	/** Background and foreground color used by Texture preview view ports. */
 	UPROPERTY(config, EditAnywhere, Category=Background)
 	FColor BackgroundColor;
diff --git a/Engine/Source/Editor/TextureEditor/Private/Menus/TextureEditorViewOptionsMenu.h b/Engine/Source/Editor/TextureEditor/Private/Menus/TextureEditorViewOptionsMenu.h
--- a/Engine/Source/Editor/TextureEditor/Private/Menus/TextureEditorViewOptionsMenu.h
+++ b/Engine/Source/Editor/TextureEditor/Private/Menus/TextureEditorViewOptionsMenu.h
@@ -20,7 +20,7 @@ public:
 	 *
 	 * @param MenuBuilder The builder for the menu that owns this menu.
 	 */
-	static void MakeMenu( FMenuBuilder& MenuBuilder, bool bIsVolumeTexture )
+	static void MakeMenu(FMenuBuilder& MenuBuilder, bool bIsVolumeTexture, bool bIsCubemapTexture)
 	{
 		// view port options
 		MenuBuilder.BeginSection("ViewportSection", LOCTEXT("ViewportSectionHeader", "Viewport Options"));
@@ -33,12 +33,12 @@ public:
 				FNewMenuDelegate::CreateStatic(&FTextureEditorViewOptionsMenu::GenerateBackgroundMenuContent)
 			);
 
-			if (bIsVolumeTexture)
+			if (bIsVolumeTexture || bIsCubemapTexture)
 			{
 				MenuBuilder.AddSubMenu(
 					LOCTEXT("ViewMode", "View Mode"),
 					LOCTEXT("ViewModeTooltip", "Set the view mode"),
-					FNewMenuDelegate::CreateStatic(&FTextureEditorViewOptionsMenu::GenerateVolumeDisplayModeMenuContent)
+					FNewMenuDelegate::CreateStatic(bIsVolumeTexture ? &FTextureEditorViewOptionsMenu::GenerateVolumeViewModeMenuContent : &FTextureEditorViewOptionsMenu::GenerateCubemapViewModeMenuContent)
 				);
 			}
 
@@ -65,15 +65,26 @@ protected:
 	}
 
 	/**
-	 * Creates the 'Display Mode' sub-menu.
+	 * Creates the 'View Mode' sub-menu for volume textures.
 	 *
 	 * @param MenuBuilder The builder for the menu that owns this menu.
 	 */
-	static void GenerateVolumeDisplayModeMenuContent( FMenuBuilder& MenuBuilder )
+	static void GenerateVolumeViewModeMenuContent( FMenuBuilder& MenuBuilder )
 	{
 		MenuBuilder.AddMenuEntry(FTextureEditorCommands::Get().DepthSlices);
 		MenuBuilder.AddMenuEntry(FTextureEditorCommands::Get().TraceIntoVolume);
 	}
+
+	/**
+	 * Creates the 'View Mode' sub-menu for cubemap textures.
+	 *
+	 * @param MenuBuilder The builder for the menu that owns this menu.
+	 */
+	static void GenerateCubemapViewModeMenuContent(FMenuBuilder& MenuBuilder)
+	{
+		MenuBuilder.AddMenuEntry(FTextureEditorCommands::Get().Cubemap2DView);
+		MenuBuilder.AddMenuEntry(FTextureEditorCommands::Get().Cubemap3DView);
+	}
 };
 
 
diff --git a/Engine/Source/Editor/TextureEditor/Private/Models/TextureEditorCommands.cpp b/Engine/Source/Editor/TextureEditor/Private/Models/TextureEditorCommands.cpp
--- a/Engine/Source/Editor/TextureEditor/Private/Models/TextureEditorCommands.cpp
+++ b/Engine/Source/Editor/TextureEditor/Private/Models/TextureEditorCommands.cpp
@@ -23,6 +23,9 @@ void FTextureEditorCommands::RegisterCommands()
 	UI_COMMAND(DepthSlices, "Depth Slices", "Shows depth slice side by side", EUserInterfaceActionType::RadioButton, FInputChord());
 	UI_COMMAND(TraceIntoVolume, "Trace Into Volume", "Traces into the volume, accumuling opacity / color", EUserInterfaceActionType::RadioButton, FInputChord());
 
+	UI_COMMAND(Cubemap2DView, "2D View", "Displays LongLat projection or individual faces of the cubemap", EUserInterfaceActionType::RadioButton, FInputChord());
+	UI_COMMAND(Cubemap3DView, "3D View", "Displays 3D rendered cubemap", EUserInterfaceActionType::RadioButton, FInputChord());
+
 	UI_COMMAND(CompressNow, "Compress", "Compress the texture", EUserInterfaceActionType::Button, FInputChord());
 	UI_COMMAND(Reimport, "Reimport", "Reimports the texture from file", EUserInterfaceActionType::Button, FInputChord());
 	UI_COMMAND(Settings, "Advanced Settings...", "Opens the advanced settings for the texture editor", EUserInterfaceActionType::Button, FInputChord());
diff --git a/Engine/Source/Editor/TextureEditor/Private/Models/TextureEditorCommands.h b/Engine/Source/Editor/TextureEditor/Private/Models/TextureEditorCommands.h
--- a/Engine/Source/Editor/TextureEditor/Private/Models/TextureEditorCommands.h
+++ b/Engine/Source/Editor/TextureEditor/Private/Models/TextureEditorCommands.h
@@ -70,6 +70,14 @@ public:
 	TSharedPtr<FUICommandInfo> TraceIntoVolume;
 // End - Volume Texture Specifics
 
+// Begin - Cubemap Texture Specifics
+	/** For cubemap texture, displays LongLat projection or individual faces of the cubemap */
+	TSharedPtr<FUICommandInfo> Cubemap2DView;
+
+	/** For cubemap texture, displays 3D rendered cubemap */
+	TSharedPtr<FUICommandInfo> Cubemap3DView;
+// End - Cubemap Texture Specifics
+
 	/** If enabled, a border is drawn around the texture */
 	TSharedPtr<FUICommandInfo> TextureBorder;
 
diff --git a/Engine/Source/Editor/TextureEditor/Private/Models/TextureEditorViewportClient.cpp b/Engine/Source/Editor/TextureEditor/Private/Models/TextureEditorViewportClient.cpp
--- a/Engine/Source/Editor/TextureEditor/Private/Models/TextureEditorViewportClient.cpp
+++ b/Engine/Source/Editor/TextureEditor/Private/Models/TextureEditorViewportClient.cpp
@@ -112,8 +112,22 @@ void FTextureEditorViewportClient::Draw(FViewport* Viewport, FCanvas* Canvas)
 	{
 		if (TextureCube || TextureCubeArray || RTTextureCube)
 		{
-			const float FaceIndex = (float)TextureEditorPinned->GetFace();
-			BatchedElementParameters = new FMipLevelBatchedElementParameters(MipLevel, SliceIndex, FaceIndex, TextureCubeArray != nullptr, false);
+			const int32 FaceIndex = TextureEditorPinned->GetFace();
+			const FVector4f& IdentityW = FVector4f(FVector3f::ZeroVector, 1);
+			// when the face index is not specified, generate a view matrix based on the tracked orientation in the texture editor
+			// (assuming that identity matrix displays the face 0 properly oriented in space),
+			// otherwise generate a view matrix which displays the selected face using DDS orientation and face order
+			// (positive x, negative x, positive y, negative y, positive z, negative z)
+			const FMatrix44f& ViewMatrix =
+				FaceIndex < 0 ? FMatrix44f(FRotationMatrix::Make(TextureEditorPinned->GetCubemapOrientation())) :
+				FaceIndex == 0 ? FMatrix44f(FVector3f::ForwardVector, FVector3f::DownVector, FVector3f::RightVector, IdentityW) :
+				FaceIndex == 1 ? FMatrix44f(FVector3f::BackwardVector, FVector3f::UpVector, FVector3f::RightVector, IdentityW) :
+				FaceIndex == 2 ? FMatrix44f(FVector3f::RightVector, FVector3f::ForwardVector, FVector3f::DownVector, IdentityW) :
+				FaceIndex == 3 ? FMatrix44f(FVector3f::LeftVector, FVector3f::ForwardVector, FVector3f::UpVector, IdentityW) :
+				FaceIndex == 4 ? FMatrix44f(FVector3f::UpVector, FVector3f::ForwardVector, FVector3f::RightVector, IdentityW) :
+				FMatrix44f(FVector3f::DownVector, FVector3f::BackwardVector, FVector3f::RightVector, IdentityW);
+			const bool bShowLongLatUnwrap = Settings.CubemapViewMode == TextureEditorCubemapViewMode_2DView && FaceIndex < 0;
+			BatchedElementParameters = new FMipLevelBatchedElementParameters(MipLevel, SliceIndex, TextureCubeArray != nullptr, ViewMatrix, bShowLongLatUnwrap, false);
 		}
 		else if (VolumeTexture)
 		{
@@ -306,6 +320,20 @@ bool FTextureEditorViewportClient::InputAxis(FViewport* Viewport, int32 Controll
 
 			TextureEditorPtr.Pin()->SetVolumeOrientation((FRotationMatrix::Make(DeltaRotator) * FRotationMatrix::Make(TextureEditorPtr.Pin()->GetVolumeOrientation())).Rotator());
 		}
+		else if (IsTextureUsingCubemapOrientation() && Viewport->KeyState(EKeys::LeftMouseButton))
+		{
+			FRotator Orienation = TextureEditorPtr.Pin()->GetCubemapOrientation();
+			const float RotationSpeed = .2f;
+			if (Key == EKeys::MouseY)
+			{
+				Orienation.Pitch += Delta * RotationSpeed;
+			}
+			else
+			{
+				Orienation.Yaw += Delta * RotationSpeed;
+			}
+			TextureEditorPtr.Pin()->SetCubemapOrientation(Orienation);
+		}
 		else if (ShouldUseMousePanning(Viewport))
 		{
 			TSharedPtr<STextureEditorViewport> EditorViewport = TextureEditorViewportPtr.Pin();
@@ -334,6 +362,14 @@ bool FTextureEditorViewportClient::InputAxis(FViewport* Viewport, int32 Controll
 	return false;
 }
 
+bool FTextureEditorViewportClient::IsTextureUsingCubemapOrientation() const
+{
+	TSharedPtr<ITextureEditorToolkit> TextureEditorPinned = TextureEditorPtr.Pin();
+	UTexture* Texture = TextureEditorPinned->GetTexture();
+	return Texture && (Texture->IsA(UTextureCube::StaticClass()) || Texture->IsA(UTextureRenderTargetCube::StaticClass()) || Texture->IsA(UTextureCubeArray::StaticClass())) &&
+		TextureEditorPinned->GetFace() < 0 && GetDefault<UTextureEditorSettings>()->CubemapViewMode == TextureEditorCubemapViewMode_3DView;
+}
+
 bool FTextureEditorViewportClient::ShouldUseMousePanning(FViewport* Viewport) const
 {
 	if (!IsTextureUsingVolumeOrientation(TextureEditorPtr.Pin()->GetTexture()) && Viewport->KeyState(EKeys::RightMouseButton))
diff --git a/Engine/Source/Editor/TextureEditor/Private/Models/TextureEditorViewportClient.h b/Engine/Source/Editor/TextureEditor/Private/Models/TextureEditorViewportClient.h
--- a/Engine/Source/Editor/TextureEditor/Private/Models/TextureEditorViewportClient.h
+++ b/Engine/Source/Editor/TextureEditor/Private/Models/TextureEditorViewportClient.h
@@ -59,6 +59,9 @@ private:
 	/** TRUE if right clicking and dragging for panning a texture 2D */
 	bool ShouldUseMousePanning(FViewport* Viewport) const;
 
+	/** Returns true if the texture is using cubemap orientation */
+	bool IsTextureUsingCubemapOrientation() const;
+
 private:
 	/** Pointer back to the Texture editor tool that owns us */
 	TWeakPtr<ITextureEditorToolkit> TextureEditorPtr;
diff --git a/Engine/Source/Editor/TextureEditor/Private/TextureEditorToolkit.cpp b/Engine/Source/Editor/TextureEditor/Private/TextureEditorToolkit.cpp
--- a/Engine/Source/Editor/TextureEditor/Private/TextureEditorToolkit.cpp
+++ b/Engine/Source/Editor/TextureEditor/Private/TextureEditorToolkit.cpp
@@ -122,6 +122,7 @@ FTextureEditorToolkit::FTextureEditorToolkit()
 	: Texture(nullptr)
 	, VolumeOpacity(1.f)
 	, VolumeOrientation(90, 0, -90)
+	, CubemapOrientation(0, 0, 0)
 {
 }
 
@@ -726,9 +727,11 @@ void FTextureEditorToolkit::PopulateQuickInfo( )
 	SizeOptions.UseGrouping = false;
 	SizeOptions.MaximumFractionalDigits = 0;
 
-	// Cubes are previewed as longlat unwrapped 2D textures in case when face index is not specified.
+	UTextureEditorSettings& Settings = *GetMutableDefault<UTextureEditorSettings>();
+
+	// Cubes are previewed as longlat unwrapped 2D textures in 2D view mode when face index is not specified.
 	// These have 2x the width of a cube face.
-	if (IsCubeTexture() && GetFace() < 0)
+	if (IsCubeTexture() && Settings.CubemapViewMode == TextureEditorCubemapViewMode_2DView && GetFace() < 0)
 	{
 		PreviewEffectiveTextureWidth *= 2;
 	}
@@ -748,7 +751,6 @@ void FTextureEditorToolkit::PopulateQuickInfo( )
 		CurrentText->SetText(FText::Format( NSLOCTEXT("TextureEditor", "QuickInfo_Displayed_3x", "Displayed: {0}x{1}x{2}"), FText::AsNumber(PreviewEffectiveTextureWidth, &Options ), FText::AsNumber(PreviewEffectiveTextureHeight, &Options), FText::AsNumber(PreviewEffectiveTextureDepth, &Options)));
 		MaxInGameText->SetText(FText::Format( NSLOCTEXT("TextureEditor", "QuickInfo_MaxInGame_3x_v1", "Max In-Game: {0}x{1}x{2}"), FText::AsNumber(MaxInGameWidth, &Options), FText::AsNumber(MaxInGameHeight, &Options), FText::AsNumber(MaxInGameDepth, &Options)));
 
-		UTextureEditorSettings& Settings = *GetMutableDefault<UTextureEditorSettings>();
 		if (Settings.VolumeViewMode == ETextureEditorVolumeViewMode::TextureEditorVolumeViewMode_VolumeTrace)
 		{
 			PreviewEffectiveTextureWidth = PreviewEffectiveTextureHeight = FMath::Max(PreviewEffectiveTextureWidth, PreviewEffectiveTextureHeight);
@@ -901,6 +903,16 @@ void FTextureEditorToolkit::SetVolumeOrientation(const FRotator& InOrientation)
 	VolumeOrientation = InOrientation;
 }
 
+const FRotator& FTextureEditorToolkit::GetCubemapOrientation() const
+{
+	return CubemapOrientation;
+}
+
+void FTextureEditorToolkit::SetCubemapOrientation(const FRotator& InOrientation)
+{
+	CubemapOrientation = InOrientation;
+}
+
 /* IToolkit interface
  *****************************************************************************/
 
@@ -1030,6 +1042,20 @@ void FTextureEditorToolkit::BindCommands( )
 		FIsActionChecked::CreateSP(this, &FTextureEditorToolkit::HandleVolumeViewModeActionIsChecked, TextureEditorVolumeViewMode_VolumeTrace));
 	// End - Volume Texture Specifics
 
+	// Begin - Cubemap Texture Specifics
+	ToolkitCommands->MapAction(
+		Commands.Cubemap2DView,
+		FExecuteAction::CreateSP(this, &FTextureEditorToolkit::HandleCubemapViewModeActionExecute, TextureEditorCubemapViewMode_2DView),
+		FCanExecuteAction(),
+		FIsActionChecked::CreateSP(this, &FTextureEditorToolkit::HandleCubemapViewModeActionIsChecked, TextureEditorCubemapViewMode_2DView));
+
+	ToolkitCommands->MapAction(
+		Commands.Cubemap3DView,
+		FExecuteAction::CreateSP(this, &FTextureEditorToolkit::HandleCubemapViewModeActionExecute, TextureEditorCubemapViewMode_3DView),
+		FCanExecuteAction(),
+		FIsActionChecked::CreateSP(this, &FTextureEditorToolkit::HandleCubemapViewModeActionIsChecked, TextureEditorCubemapViewMode_3DView));
+	// End - Cubemap Texture Specifics
+
 	ToolkitCommands->MapAction(
 		Commands.TextureBorder,
 		FExecuteAction::CreateSP(this, &FTextureEditorToolkit::HandleTextureBorderActionExecute),
@@ -1926,7 +1952,7 @@ TSharedRef<SWidget> FTextureEditorToolkit::OnGenerateMipMapLevelMenu()
 TSharedRef<SWidget> FTextureEditorToolkit::OnGenerateSettingsMenu()
 {
 	FMenuBuilder MenuBuilder(true, ToolkitCommands);
-	FTextureEditorViewOptionsMenu::MakeMenu(MenuBuilder, IsVolumeTexture());
+	FTextureEditorViewOptionsMenu::MakeMenu(MenuBuilder, IsVolumeTexture(), IsCubeTexture());
 
 	return MenuBuilder.MakeWidget();
 }
@@ -2040,7 +2066,7 @@ bool FTextureEditorToolkit::HandleCheckeredBackgroundActionIsChecked( ETextureEd
 	return (Background == Settings.Background);
 }
 
-// Callback for toggling the volume display action.
+// Callback for toggling the volume view action.
 void FTextureEditorToolkit::HandleVolumeViewModeActionExecute(ETextureEditorVolumeViewMode InViewMode)
 {
 	UTextureEditorSettings& Settings = *GetMutableDefault<UTextureEditorSettings>();
@@ -2048,7 +2074,7 @@ void FTextureEditorToolkit::HandleVolumeViewModeActionExecute(ETextureEditorVolu
 	Settings.PostEditChange();
 }
 
-// Callback for getting the checked state of the volume display action.
+// Callback for getting the checked state of the volume view action.
 bool FTextureEditorToolkit::HandleVolumeViewModeActionIsChecked(ETextureEditorVolumeViewMode InViewMode)
 {
 	const UTextureEditorSettings& Settings = *GetDefault<UTextureEditorSettings>();
@@ -2056,6 +2082,20 @@ bool FTextureEditorToolkit::HandleVolumeViewModeActionIsChecked(ETextureEditorVo
 	return (InViewMode == Settings.VolumeViewMode);
 }
 
+// Callback for toggling the cubemap view action.
+void FTextureEditorToolkit::HandleCubemapViewModeActionExecute(ETextureEditorCubemapViewMode InViewMode)
+{
+	UTextureEditorSettings& Settings = *GetMutableDefault<UTextureEditorSettings>();
+	Settings.CubemapViewMode = InViewMode;
+	Settings.PostEditChange();
+}
+
+// Callback for getting the checked state of the cubemap view action.
+bool FTextureEditorToolkit::HandleCubemapViewModeActionIsChecked(ETextureEditorCubemapViewMode InViewMode)
+{
+	UTextureEditorSettings& Settings = *GetMutableDefault<UTextureEditorSettings>();
+	return Settings.CubemapViewMode == InViewMode;
+}
 
 void FTextureEditorToolkit::HandleCompressNowActionExecute( )
 {
diff --git a/Engine/Source/Editor/TextureEditor/Private/TextureEditorToolkit.h b/Engine/Source/Editor/TextureEditor/Private/TextureEditorToolkit.h
--- a/Engine/Source/Editor/TextureEditor/Private/TextureEditorToolkit.h
+++ b/Engine/Source/Editor/TextureEditor/Private/TextureEditorToolkit.h
@@ -85,6 +85,8 @@ public:
 	virtual void SetVolumeOpacity( float VolumeOpacity ) override;
 	virtual const FRotator& GetVolumeOrientation( ) const override;
 	virtual void SetVolumeOrientation( const FRotator& InOrientation ) override;
+	virtual const FRotator& GetCubemapOrientation() const override;
+	virtual void SetCubemapOrientation(const FRotator& InOrientation) override;
 	virtual int32 GetExposureBias() const override
 	{
 		return ExposureBias;
@@ -194,12 +196,18 @@ private:
 	// Callback for getting the checked state of the Checkered Background action.
 	bool HandleCheckeredBackgroundActionIsChecked( ETextureEditorBackgrounds Background );
 
-	// Callback for toggling the volume display action.
+	// Callback for toggling the volume view action.
 	void HandleVolumeViewModeActionExecute( ETextureEditorVolumeViewMode InViewMode );
 
-	// Callback for getting the checked state of the volume display action.
+	// Callback for getting the checked state of the volume view action.
 	bool HandleVolumeViewModeActionIsChecked( ETextureEditorVolumeViewMode InViewMode );
 
+	// Callback for toggling the cubemap view action.
+	void HandleCubemapViewModeActionExecute(ETextureEditorCubemapViewMode InViewMode);
+
+	// Callback for getting the checked state of the cubemap view action.
+	bool HandleCubemapViewModeActionIsChecked(ETextureEditorCubemapViewMode InViewMode);
+
 	// Callback for toggling the Compress Now action.
 	void HandleCompressNowActionExecute( );
 
@@ -465,6 +473,9 @@ private:
 	// For volume texture, the orientation when tracing.
 	FRotator VolumeOrientation;
 
+	// For cubemap texture, the orientation when previewing the texture as a 3d rendered cube.
+	FRotator CubemapOrientation;
+
 	bool bIsVolumeTexture;
 
 	// Objects and callbacks for the custom encoding settings checkbox
diff --git a/Engine/Source/Editor/TextureEditor/Public/Interfaces/ITextureEditorToolkit.h b/Engine/Source/Editor/TextureEditor/Public/Interfaces/ITextureEditorToolkit.h
--- a/Engine/Source/Editor/TextureEditor/Public/Interfaces/ITextureEditorToolkit.h
+++ b/Engine/Source/Editor/TextureEditor/Public/Interfaces/ITextureEditorToolkit.h
@@ -49,6 +49,8 @@ public:
 	virtual void SetVolumeOpacity( float VolumeOpacity ) = 0;
 	virtual const FRotator& GetVolumeOrientation( ) const = 0;
 	virtual void SetVolumeOrientation( const FRotator& InOrientation ) = 0;
+	virtual const FRotator& GetCubemapOrientation() const = 0;
+	virtual void SetCubemapOrientation(const FRotator& InOrientation) = 0;
 	virtual int32 GetExposureBias() const = 0;
 	virtual bool IsVolumeTexture() const = 0;
 public:
diff --git a/Engine/Source/Editor/UnrealEd/Private/ThumbnailRendering/TextureThumbnailRenderer.cpp b/Engine/Source/Editor/UnrealEd/Private/ThumbnailRendering/TextureThumbnailRenderer.cpp
--- a/Engine/Source/Editor/UnrealEd/Private/ThumbnailRendering/TextureThumbnailRenderer.cpp
+++ b/Engine/Source/Editor/UnrealEd/Private/ThumbnailRendering/TextureThumbnailRenderer.cpp
@@ -72,7 +72,7 @@ void UTextureThumbnailRenderer::Draw(UObject* Object, int32 X, int32 Y, uint32 W
 		if(TextureCube || TextureCubeArray || RTTextureCube)
 		{
 			// is released by the render thread when it was rendered
-			BatchedElementParameters = new FMipLevelBatchedElementParameters((float)0, (float)-1, (float)-1, TextureCubeArray != nullptr);
+			BatchedElementParameters = new FMipLevelBatchedElementParameters((float)0, (float)-1, TextureCubeArray != nullptr, FMatrix44f::Identity, true, false);
 			
 			// If the thumbnail is square then make it 2:1 for cubes.
 			if(Width == Height)
diff --git a/Engine/Source/Runtime/Engine/Private/CubemapUnwrapUtils.cpp b/Engine/Source/Runtime/Engine/Private/CubemapUnwrapUtils.cpp
--- a/Engine/Source/Runtime/Engine/Private/CubemapUnwrapUtils.cpp
+++ b/Engine/Source/Runtime/Engine/Private/CubemapUnwrapUtils.cpp
@@ -36,7 +36,7 @@ namespace CubemapHelpers
 	bool GenerateLongLatUnwrap(const FTextureResource* TextureResource, const uint32 AxisDimenion, const EPixelFormat SourcePixelFormat, TArray64<uint8>& BitsOUT, FIntPoint& SizeOUT, EPixelFormat& FormatOUT)
 	{
 		TRefCountPtr<FBatchedElementParameters> BatchedElementParameters;
-		BatchedElementParameters = new FMipLevelBatchedElementParameters((float)0, (float)-1, (float)-1, false, true);
+		BatchedElementParameters = new FMipLevelBatchedElementParameters((float)0, (float)-1, false, FMatrix44f::Identity, true, true);
 		const FIntPoint LongLatDimensions(AxisDimenion * 2, AxisDimenion);
 
 		// If the source format is 8 bit per channel or less then select a LDR target format.
@@ -128,6 +128,8 @@ FCubemapTexturePropertiesPS::FCubemapTexturePropertiesPS(const ShaderMetaType::C
 	PackedProperties0.Bind(Initializer.ParameterMap, TEXT("PackedProperties0"));
 	Gamma.Bind(Initializer.ParameterMap, TEXT("Gamma"));
 	NumSlices.Bind(Initializer.ParameterMap, TEXT("NumSlices"));
+	SliceIndex.Bind(Initializer.ParameterMap, TEXT("SliceIndex"));
+	ViewMatrix.Bind(Initializer.ParameterMap, TEXT("ViewMatrix"));
 }
 
 bool FCubemapTexturePropertiesPS::ShouldCompilePermutation(const FGlobalShaderPermutationParameters& Parameters)
@@ -145,24 +147,26 @@ bool FCubemapTexturePropertiesPS::ShouldCompilePermutation(const FGlobalShaderPe
 	return true;
 }
 
-void FCubemapTexturePropertiesPS::SetParameters(FRHICommandList& RHICmdList, const FTexture* Texture, const FMatrix& ColorWeightsValue, float MipLevel, float SliceIndex, float FaceIndex, float GammaValue, bool bIsTextureCubeArray)
+void FCubemapTexturePropertiesPS::SetParameters(FRHICommandList& RHICmdList, const FTexture* InTexture, const FMatrix& InColorWeightsValue, float InMipLevel, float InSliceIndex, bool bInIsTextureCubeArray, const FMatrix44f& InViewMatrix, bool bInShowLongLatUnwrap, float InGammaValue)
 {
 	FRHIPixelShader* ShaderRHI = RHICmdList.GetBoundPixelShader();
 
-	SetTextureParameter(RHICmdList, ShaderRHI, CubeTexture, CubeTextureSampler, Texture);
+	SetTextureParameter(RHICmdList, ShaderRHI, CubeTexture, CubeTextureSampler, InTexture);
 
-	FVector4f PackedProperties0Value(MipLevel, SliceIndex, FaceIndex, 0);
+	FVector4f PackedProperties0Value(InMipLevel, bInShowLongLatUnwrap ? 1.0f : -1.0f, 0, 0);
 	SetShaderValue(RHICmdList, ShaderRHI, PackedProperties0, PackedProperties0Value);
-	SetShaderValue(RHICmdList, ShaderRHI, ColorWeights, (FMatrix44f)ColorWeightsValue);
-	SetShaderValue(RHICmdList, ShaderRHI, Gamma, GammaValue);
+	SetShaderValue(RHICmdList, ShaderRHI, ColorWeights, (FMatrix44f)InColorWeightsValue);
+	SetShaderValue(RHICmdList, ShaderRHI, Gamma, InGammaValue);
+	SetShaderValue(RHICmdList, RHICmdList.GetBoundPixelShader(), ViewMatrix, InViewMatrix);
 
-	// store slice count for the texture cube array
-	if (bIsTextureCubeArray)
+	// store slice count and selected slice index for the texture cube array
+	if (bInIsTextureCubeArray)
 	{
 		// GetSizeZ() returns the total number of slices stored in the platform data
 		// for a TextureCube array this value is equal to the size of the array multiplied by 6
-		const float NumSlicesData = (float)(Texture ? FMath::Max((int32)Texture->GetSizeZ() / 6, 1) : 1);
+		const float NumSlicesData = (float)(InTexture ? FMath::Max((int32)InTexture->GetSizeZ() / 6, 1) : 1);
 		SetShaderValue(RHICmdList, ShaderRHI, NumSlices, NumSlicesData);
+		SetShaderValue(RHICmdList, ShaderRHI, SliceIndex, InSliceIndex);
 	}
 }
 
@@ -187,7 +191,7 @@ void FMipLevelBatchedElementParameters::BindShaders(FRHICommandList& RHICmdList,
 	SetGraphicsPipelineState(RHICmdList, GraphicsPSOInit, 0);
 
 	VertexShader->SetParameters(RHICmdList, InTransform);
-	PixelShader->SetParameters(RHICmdList, Texture, ColorWeights, MipLevel, SliceIndex, FaceIndex, InGamma, bIsTextureCubeArray);
+	PixelShader->SetParameters(RHICmdList, Texture, ColorWeights, MipLevel, SliceIndex, bIsTextureCubeArray, ViewMatrix, bShowLongLatUnwrap, InGamma);
 }
 
 void FIESLightProfilePS::SetParameters( FRHICommandList& RHICmdList, const FTexture* Texture, float InBrightnessInLumens )
diff --git a/Engine/Source/Runtime/Engine/Public/CubemapUnwrapUtils.h b/Engine/Source/Runtime/Engine/Public/CubemapUnwrapUtils.h
--- a/Engine/Source/Runtime/Engine/Public/CubemapUnwrapUtils.h
+++ b/Engine/Source/Runtime/Engine/Public/CubemapUnwrapUtils.h
@@ -79,7 +79,7 @@ public:
 	FCubemapTexturePropertiesPS(const ShaderMetaType::CompiledShaderInitializerType& Initializer);
 
 	static bool ShouldCompilePermutation(const FGlobalShaderPermutationParameters& Parameters);
-	void SetParameters(FRHICommandList& RHICmdList, const FTexture* Texture, const FMatrix& ColorWeightsValue, float MipLevel, float SliceIndex, float FaceIndex, float GammaValue, bool bIsTextureCubeArray);
+	void SetParameters(FRHICommandList& RHICmdList, const FTexture* InTexture, const FMatrix& InColorWeightsValue, float InMipLevel, float InSliceIndex, bool bInIsTextureCubeArray, const FMatrix44f& InViewMatrix, bool bInShowLongLatUnwrap, float InGammaValue);
 
 private:
 	LAYOUT_FIELD(FShaderResourceParameter, CubeTexture);
@@ -88,17 +88,20 @@ private:
 	LAYOUT_FIELD(FShaderParameter, ColorWeights);
 	LAYOUT_FIELD(FShaderParameter, Gamma);
 	LAYOUT_FIELD(FShaderParameter, NumSlices);
+	LAYOUT_FIELD(FShaderParameter, SliceIndex);
+	LAYOUT_FIELD(FShaderParameter, ViewMatrix);
 };
 
 
 class ENGINE_API FMipLevelBatchedElementParameters : public FBatchedElementParameters
 {
 public:
-	FMipLevelBatchedElementParameters(float InMipLevel, float InSliceIndex, float InFaceIndex, bool bInIsTextureCubeArray, bool bInHDROutput = false)
+	FMipLevelBatchedElementParameters(float InMipLevel, float InSliceIndex, bool bInIsTextureCubeArray, const FMatrix44f& InViewMatrix, bool bInShowLongLatUnwrap, bool bInHDROutput)
 		: bHDROutput(bInHDROutput)
 		, MipLevel(InMipLevel)
 		, SliceIndex(InSliceIndex)
-		, FaceIndex(InFaceIndex)
+		, ViewMatrix(InViewMatrix)
+		, bShowLongLatUnwrap(bInShowLongLatUnwrap)
 		, bIsTextureCubeArray(bInIsTextureCubeArray)
 	{
 	}
@@ -112,7 +115,8 @@ private:
 	/** Parameters that need to be passed to the shader */
 	float MipLevel;
 	float SliceIndex;
-	float FaceIndex;
+	FMatrix44f ViewMatrix;
+	bool bShowLongLatUnwrap;
 
 	/** Parameters that are used to select a shader permutation */
 	bool bIsTextureCubeArray;
