Update TextureEditor viewport to use the newer version of InputKey and InputAxis.

diff --git a/Engine/Source/Editor/TextureEditor/Private/Models/TextureEditorViewportClient.cpp b/Engine/Source/Editor/TextureEditor/Private/Models/TextureEditorViewportClient.cpp
--- a/Engine/Source/Editor/TextureEditor/Private/Models/TextureEditorViewportClient.cpp
+++ b/Engine/Source/Editor/TextureEditor/Private/Models/TextureEditorViewportClient.cpp
@@ -271,23 +271,23 @@ void FTextureEditorViewportClient::Draw(FViewport* Viewport, FCanvas* Canvas)
 }
 
 
-bool FTextureEditorViewportClient::InputKey(FViewport* Viewport, int32 ControllerId, FKey Key, EInputEvent Event, float AmountDepressed, bool Gamepad)
+bool FTextureEditorViewportClient::InputKey(const FInputKeyEventArgs& InEventArgs)
 {
-	if (Event == IE_Pressed)
+	if (InEventArgs.Event == IE_Pressed)
 	{
-		if (Key == EKeys::MouseScrollUp)
+		if (InEventArgs.Key == EKeys::MouseScrollUp)
 		{
 			TextureEditorPtr.Pin()->ZoomIn();
 
 			return true;
 		}
-		else if (Key == EKeys::MouseScrollDown)
+		else if (InEventArgs.Key == EKeys::MouseScrollDown)
 		{
 			TextureEditorPtr.Pin()->ZoomOut();
 
 			return true;
 		}
-		else if (Key == EKeys::RightMouseButton)
+		else if (InEventArgs.Key == EKeys::RightMouseButton)
 		{
 			TextureEditorPtr.Pin()->SetVolumeOrientation(FRotator(90, 0, -90));
 		}
@@ -300,7 +300,7 @@ bool IsTextureUsingVolumeOrientation(UTexture* Texture)
 	return Texture && (Cast<UVolumeTexture>(Texture) || Cast<UTextureRenderTargetVolume>(Texture));
 }
 
-bool FTextureEditorViewportClient::InputAxis(FViewport* Viewport, int32 ControllerId, FKey Key, float Delta, float DeltaTime, int32 NumSamples, bool bGamepad)
+bool FTextureEditorViewportClient::InputAxis(FViewport* Viewport, FInputDeviceId DeviceId, FKey Key, float Delta, float DeltaTime, int32 NumSamples, bool bGamepad)
 {
 	if (Key == EKeys::MouseX || Key == EKeys::MouseY)
 	{
diff --git a/Engine/Source/Editor/TextureEditor/Private/Models/TextureEditorViewportClient.h b/Engine/Source/Editor/TextureEditor/Private/Models/TextureEditorViewportClient.h
--- a/Engine/Source/Editor/TextureEditor/Private/Models/TextureEditorViewportClient.h
+++ b/Engine/Source/Editor/TextureEditor/Private/Models/TextureEditorViewportClient.h
@@ -23,8 +23,8 @@ public:
 
 	/** FViewportClient interface */
 	virtual void Draw(FViewport* Viewport, FCanvas* Canvas) override;
-	virtual bool InputKey(FViewport* Viewport, int32 ControllerId, FKey Key, EInputEvent Event, float AmountDepressed = 1.0f, bool bGamepad = false) override;
-	virtual bool InputAxis(FViewport* Viewport,int32 ControllerId,FKey Key,float Delta,float DeltaTime,int32 NumSamples=1,bool bGamepad=false) override;
+	virtual bool InputKey(const FInputKeyEventArgs& InEventArgs) override;
+	virtual bool InputAxis(FViewport* Viewport, FInputDeviceId DeviceId, FKey Key, float Delta, float DeltaTime, int32 NumSamples = 1, bool bGamepad = false) override;
 	virtual bool InputGesture(FViewport* Viewport, EGestureEvent GestureType, const FVector2D& GestureDelta, bool bIsDirectionInvertedFromDevice) override;
 	virtual UWorld* GetWorld() const override { return nullptr; }
 	virtual EMouseCursor::Type GetCursor(FViewport* Viewport, int32 X, int32 Y) override;
