Move view mode control for volume textures and cubemaps into the main panel of the texture editor.

diff --git a/Engine/Source/Editor/TextureEditor/Private/Menus/TextureEditorViewOptionsMenu.h b/Engine/Source/Editor/TextureEditor/Private/Menus/TextureEditorViewOptionsMenu.h
--- a/Engine/Source/Editor/TextureEditor/Private/Menus/TextureEditorViewOptionsMenu.h
+++ b/Engine/Source/Editor/TextureEditor/Private/Menus/TextureEditorViewOptionsMenu.h
@@ -20,7 +20,7 @@ public:
 	 *
 	 * @param MenuBuilder The builder for the menu that owns this menu.
 	 */
-	static void MakeMenu(FMenuBuilder& MenuBuilder, bool bIsVolumeTexture, bool bIsCubemapTexture)
+	static void MakeMenu(FMenuBuilder& MenuBuilder)
 	{
 		// view port options
 		MenuBuilder.BeginSection("ViewportSection", LOCTEXT("ViewportSectionHeader", "Viewport Options"));
@@ -33,15 +33,6 @@ public:
 				FNewMenuDelegate::CreateStatic(&FTextureEditorViewOptionsMenu::GenerateBackgroundMenuContent)
 			);
 
-			if (bIsVolumeTexture || bIsCubemapTexture)
-			{
-				MenuBuilder.AddSubMenu(
-					LOCTEXT("ViewMode", "View Mode"),
-					LOCTEXT("ViewModeTooltip", "Set the view mode"),
-					FNewMenuDelegate::CreateStatic(bIsVolumeTexture ? &FTextureEditorViewOptionsMenu::GenerateVolumeViewModeMenuContent : &FTextureEditorViewOptionsMenu::GenerateCubemapViewModeMenuContent)
-				);
-			}
-
 			MenuBuilder.AddMenuEntry(FTextureEditorCommands::Get().TextureBorder);
 		}
 		MenuBuilder.EndSection();
@@ -64,27 +55,6 @@ protected:
 		MenuBuilder.AddMenuEntry(FTextureEditorCommands::Get().SolidBackground);
 	}
 
-	/**
-	 * Creates the 'View Mode' sub-menu for volume textures.
-	 *
-	 * @param MenuBuilder The builder for the menu that owns this menu.
-	 */
-	static void GenerateVolumeViewModeMenuContent( FMenuBuilder& MenuBuilder )
-	{
-		MenuBuilder.AddMenuEntry(FTextureEditorCommands::Get().DepthSlices);
-		MenuBuilder.AddMenuEntry(FTextureEditorCommands::Get().TraceIntoVolume);
-	}
-
-	/**
-	 * Creates the 'View Mode' sub-menu for cubemap textures.
-	 *
-	 * @param MenuBuilder The builder for the menu that owns this menu.
-	 */
-	static void GenerateCubemapViewModeMenuContent(FMenuBuilder& MenuBuilder)
-	{
-		MenuBuilder.AddMenuEntry(FTextureEditorCommands::Get().Cubemap2DView);
-		MenuBuilder.AddMenuEntry(FTextureEditorCommands::Get().Cubemap3DView);
-	}
 };
 
 
diff --git a/Engine/Source/Editor/TextureEditor/Private/Models/TextureEditorCommands.cpp b/Engine/Source/Editor/TextureEditor/Private/Models/TextureEditorCommands.cpp
--- a/Engine/Source/Editor/TextureEditor/Private/Models/TextureEditorCommands.cpp
+++ b/Engine/Source/Editor/TextureEditor/Private/Models/TextureEditorCommands.cpp
@@ -20,12 +20,6 @@ void FTextureEditorCommands::RegisterCommands()
 	UI_COMMAND(SolidBackground, "Solid Color", "Solid color background", EUserInterfaceActionType::RadioButton, FInputChord());
 	UI_COMMAND(TextureBorder, "Draw Border", "If enabled, a border is drawn around the texture", EUserInterfaceActionType::ToggleButton, FInputChord());
 
-	UI_COMMAND(DepthSlices, "Depth Slices", "Shows depth slice side by side", EUserInterfaceActionType::RadioButton, FInputChord());
-	UI_COMMAND(TraceIntoVolume, "Trace Into Volume", "Traces into the volume, accumuling opacity / color", EUserInterfaceActionType::RadioButton, FInputChord());
-
-	UI_COMMAND(Cubemap2DView, "2D View", "Displays LongLat projection or individual faces of the cubemap", EUserInterfaceActionType::RadioButton, FInputChord());
-	UI_COMMAND(Cubemap3DView, "3D View", "Displays 3D rendered cubemap", EUserInterfaceActionType::RadioButton, FInputChord());
-
 	UI_COMMAND(CompressNow, "Compress", "Compress the texture", EUserInterfaceActionType::Button, FInputChord());
 	UI_COMMAND(Reimport, "Reimport", "Reimports the texture from file", EUserInterfaceActionType::Button, FInputChord());
 	UI_COMMAND(Settings, "Advanced Settings...", "Opens the advanced settings for the texture editor", EUserInterfaceActionType::Button, FInputChord());
diff --git a/Engine/Source/Editor/TextureEditor/Private/Models/TextureEditorCommands.h b/Engine/Source/Editor/TextureEditor/Private/Models/TextureEditorCommands.h
--- a/Engine/Source/Editor/TextureEditor/Private/Models/TextureEditorCommands.h
+++ b/Engine/Source/Editor/TextureEditor/Private/Models/TextureEditorCommands.h
@@ -62,22 +62,6 @@ public:
 	/** Sets the solid color background */
 	TSharedPtr<FUICommandInfo> SolidBackground;
 
-// Begin - Volume Texture Specifics
-	/** For volume texture, shows each depth slice side by side */
-	TSharedPtr<FUICommandInfo> DepthSlices;
-
-	/** For volume texture, shows a trace into the volume */
-	TSharedPtr<FUICommandInfo> TraceIntoVolume;
-// End - Volume Texture Specifics
-
-// Begin - Cubemap Texture Specifics
-	/** For cubemap texture, displays LongLat projection or individual faces of the cubemap */
-	TSharedPtr<FUICommandInfo> Cubemap2DView;
-
-	/** For cubemap texture, displays 3D rendered cubemap */
-	TSharedPtr<FUICommandInfo> Cubemap3DView;
-// End - Cubemap Texture Specifics
-
 	/** If enabled, a border is drawn around the texture */
 	TSharedPtr<FUICommandInfo> TextureBorder;
 
diff --git a/Engine/Source/Editor/TextureEditor/Private/Models/TextureEditorViewportClient.cpp b/Engine/Source/Editor/TextureEditor/Private/Models/TextureEditorViewportClient.cpp
--- a/Engine/Source/Editor/TextureEditor/Private/Models/TextureEditorViewportClient.cpp
+++ b/Engine/Source/Editor/TextureEditor/Private/Models/TextureEditorViewportClient.cpp
@@ -126,13 +126,13 @@ void FTextureEditorViewportClient::Draw(FViewport* Viewport, FCanvas* Canvas)
 				FaceIndex == 3 ? FMatrix44f(FVector3f::LeftVector, FVector3f::ForwardVector, FVector3f::UpVector, IdentityW) :
 				FaceIndex == 4 ? FMatrix44f(FVector3f::UpVector, FVector3f::ForwardVector, FVector3f::RightVector, IdentityW) :
 				FMatrix44f(FVector3f::DownVector, FVector3f::BackwardVector, FVector3f::RightVector, IdentityW);
-			const bool bShowLongLatUnwrap = Settings.CubemapViewMode == TextureEditorCubemapViewMode_2DView && FaceIndex < 0;
+			const bool bShowLongLatUnwrap = TextureEditorPinned->GetCubemapViewMode() == TextureEditorCubemapViewMode_2DView && FaceIndex < 0;
 			BatchedElementParameters = new FMipLevelBatchedElementParameters(MipLevel, SliceIndex, TextureCubeArray != nullptr, ViewMatrix, bShowLongLatUnwrap, false);
 		}
 		else if (VolumeTexture)
 		{
 			BatchedElementParameters = new FBatchedElementVolumeTexturePreviewParameters(
-				Settings.VolumeViewMode == TextureEditorVolumeViewMode_DepthSlices, 
+				TextureEditorPinned->GetVolumeViewMode() == TextureEditorVolumeViewMode_DepthSlices,
 				FMath::Max<int32>(VolumeTexture->GetSizeZ(), 1), 
 				MipLevel, 
 				(float)TextureEditorPinned->GetVolumeOpacity(),
@@ -142,7 +142,7 @@ void FTextureEditorViewportClient::Draw(FViewport* Viewport, FCanvas* Canvas)
 		else if (RTTextureVolume)
 		{
 			BatchedElementParameters = new FBatchedElementVolumeTexturePreviewParameters(
-				Settings.VolumeViewMode == TextureEditorVolumeViewMode_DepthSlices,
+				TextureEditorPinned->GetVolumeViewMode() == TextureEditorVolumeViewMode_DepthSlices,
 				FMath::Max<int32>(RTTextureVolume->SizeZ >> RTTextureVolume->GetCachedLODBias(), 1),
 				MipLevel,
 				(float)TextureEditorPinned->GetVolumeOpacity(),
@@ -367,7 +367,7 @@ bool FTextureEditorViewportClient::IsTextureUsingCubemapOrientation() const
 	TSharedPtr<ITextureEditorToolkit> TextureEditorPinned = TextureEditorPtr.Pin();
 	UTexture* Texture = TextureEditorPinned->GetTexture();
 	return Texture && (Texture->IsA(UTextureCube::StaticClass()) || Texture->IsA(UTextureRenderTargetCube::StaticClass()) || Texture->IsA(UTextureCubeArray::StaticClass())) &&
-		TextureEditorPinned->GetFace() < 0 && GetDefault<UTextureEditorSettings>()->CubemapViewMode == TextureEditorCubemapViewMode_3DView;
+		TextureEditorPinned->GetFace() < 0 && TextureEditorPinned->GetCubemapViewMode() == TextureEditorCubemapViewMode_3DView;
 }
 
 bool FTextureEditorViewportClient::ShouldUseMousePanning(FViewport* Viewport) const
diff --git a/Engine/Source/Editor/TextureEditor/Private/TextureEditorToolkit.cpp b/Engine/Source/Editor/TextureEditor/Private/TextureEditorToolkit.cpp
--- a/Engine/Source/Editor/TextureEditor/Private/TextureEditorToolkit.cpp
+++ b/Engine/Source/Editor/TextureEditor/Private/TextureEditorToolkit.cpp
@@ -254,9 +254,12 @@ void FTextureEditorToolkit::InitTextureEditor( const EToolkitMode::Type Mode, co
 
 	SavedCompressionSetting = false;
 
-	// Start at whatever the last used zoom mode was
+	// Start at whatever the last used zoom mode, volume view mode and cubemap view mode were
 	const UTextureEditorSettings& Settings = *GetDefault<UTextureEditorSettings>();
 	ZoomMode = Settings.ZoomMode;
+	VolumeViewMode = Settings.VolumeViewMode;
+	CubemapViewMode = Settings.CubemapViewMode;
+
 	Zoom = 1.0f;
 
 	// Register our commands. This will only register them if not previously registered
@@ -341,7 +344,7 @@ void FTextureEditorToolkit::CalculateTextureDimensions(int32& OutWidth, int32& O
 		const int32 MaxWidth = FMath::Max(TextureViewport->GetViewport()->GetSizeXY().X - 2 * BorderSize, 0);
 		const int32 MaxHeight = FMath::Max(TextureViewport->GetViewport()->GetSizeXY().Y - 2 * BorderSize, 0);
 
-		if (IsVolumeTexture() && GetMutableDefault<UTextureEditorSettings>()->VolumeViewMode == ETextureEditorVolumeViewMode::TextureEditorVolumeViewMode_VolumeTrace)
+		if (IsVolumeTexture() && GetVolumeViewMode() == ETextureEditorVolumeViewMode::TextureEditorVolumeViewMode_VolumeTrace)
 		{
 			OutWidth = OutHeight = FMath::Min(MaxWidth, MaxHeight);
 		}
@@ -727,11 +730,9 @@ void FTextureEditorToolkit::PopulateQuickInfo( )
 	SizeOptions.UseGrouping = false;
 	SizeOptions.MaximumFractionalDigits = 0;
 
-	UTextureEditorSettings& Settings = *GetMutableDefault<UTextureEditorSettings>();
-
 	// Cubes are previewed as longlat unwrapped 2D textures in 2D view mode when face index is not specified.
 	// These have 2x the width of a cube face.
-	if (IsCubeTexture() && Settings.CubemapViewMode == TextureEditorCubemapViewMode_2DView && GetFace() < 0)
+	if (IsCubeTexture() && GetCubemapViewMode() == TextureEditorCubemapViewMode_2DView && GetFace() < 0)
 	{
 		PreviewEffectiveTextureWidth *= 2;
 	}
@@ -751,7 +752,7 @@ void FTextureEditorToolkit::PopulateQuickInfo( )
 		CurrentText->SetText(FText::Format( NSLOCTEXT("TextureEditor", "QuickInfo_Displayed_3x", "Displayed: {0}x{1}x{2}"), FText::AsNumber(PreviewEffectiveTextureWidth, &Options ), FText::AsNumber(PreviewEffectiveTextureHeight, &Options), FText::AsNumber(PreviewEffectiveTextureDepth, &Options)));
 		MaxInGameText->SetText(FText::Format( NSLOCTEXT("TextureEditor", "QuickInfo_MaxInGame_3x_v1", "Max In-Game: {0}x{1}x{2}"), FText::AsNumber(MaxInGameWidth, &Options), FText::AsNumber(MaxInGameHeight, &Options), FText::AsNumber(MaxInGameDepth, &Options)));
 
-		if (Settings.VolumeViewMode == ETextureEditorVolumeViewMode::TextureEditorVolumeViewMode_VolumeTrace)
+		if (GetVolumeViewMode() == ETextureEditorVolumeViewMode::TextureEditorVolumeViewMode_VolumeTrace)
 		{
 			PreviewEffectiveTextureWidth = PreviewEffectiveTextureHeight = FMath::Max(PreviewEffectiveTextureWidth, PreviewEffectiveTextureHeight);
 		}
@@ -903,6 +904,23 @@ void FTextureEditorToolkit::SetVolumeOrientation(const FRotator& InOrientation)
 	VolumeOrientation = InOrientation;
 }
 
+ETextureEditorVolumeViewMode FTextureEditorToolkit::GetVolumeViewMode() const
+{
+	// Each texture editor keeps a local volume view mode so that it can be changed without affecting other open editors
+	return VolumeViewMode;
+}
+
+void FTextureEditorToolkit::SetVolumeViewMode(const ETextureEditorVolumeViewMode InVolumeViewMode)
+{
+	// Update our own volume view mode
+	VolumeViewMode = InVolumeViewMode;
+
+	// And also save it so it's used for new texture editors
+	UTextureEditorSettings& Settings = *GetMutableDefault<UTextureEditorSettings>();
+	Settings.VolumeViewMode = VolumeViewMode;
+	Settings.PostEditChange();
+}
+
 const FRotator& FTextureEditorToolkit::GetCubemapOrientation() const
 {
 	return CubemapOrientation;
@@ -913,6 +931,23 @@ void FTextureEditorToolkit::SetCubemapOrientation(const FRotator& InOrientation)
 	CubemapOrientation = InOrientation;
 }
 
+ETextureEditorCubemapViewMode FTextureEditorToolkit::GetCubemapViewMode() const
+{
+	// Each texture editor keeps a local cubemap view mode so that it can be changed without affecting other open editors
+	return CubemapViewMode;
+}
+
+void FTextureEditorToolkit::SetCubemapViewMode(const ETextureEditorCubemapViewMode InCubemapViewMode)
+{
+	// Update our own cubemap view mode
+	CubemapViewMode = InCubemapViewMode;
+
+	// And also save it so it's used for new texture editors
+	UTextureEditorSettings& Settings = *GetMutableDefault<UTextureEditorSettings>();
+	Settings.CubemapViewMode = CubemapViewMode;
+	Settings.PostEditChange();
+}
+
 /* IToolkit interface
  *****************************************************************************/
 
@@ -1028,34 +1063,6 @@ void FTextureEditorToolkit::BindCommands( )
 		FCanExecuteAction(),
 		FIsActionChecked::CreateSP(this, &FTextureEditorToolkit::HandleCheckeredBackgroundActionIsChecked, TextureEditorBackground_SolidColor));
 
-	// Begin - Volume Texture Specifics
-	ToolkitCommands->MapAction(
-		Commands.DepthSlices,
-		FExecuteAction::CreateSP(this, &FTextureEditorToolkit::HandleVolumeViewModeActionExecute, TextureEditorVolumeViewMode_DepthSlices),
-		FCanExecuteAction(),
-		FIsActionChecked::CreateSP(this, &FTextureEditorToolkit::HandleVolumeViewModeActionIsChecked, TextureEditorVolumeViewMode_DepthSlices));
-
-	ToolkitCommands->MapAction(
-		Commands.TraceIntoVolume,
-		FExecuteAction::CreateSP(this, &FTextureEditorToolkit::HandleVolumeViewModeActionExecute, TextureEditorVolumeViewMode_VolumeTrace),
-		FCanExecuteAction(),
-		FIsActionChecked::CreateSP(this, &FTextureEditorToolkit::HandleVolumeViewModeActionIsChecked, TextureEditorVolumeViewMode_VolumeTrace));
-	// End - Volume Texture Specifics
-
-	// Begin - Cubemap Texture Specifics
-	ToolkitCommands->MapAction(
-		Commands.Cubemap2DView,
-		FExecuteAction::CreateSP(this, &FTextureEditorToolkit::HandleCubemapViewModeActionExecute, TextureEditorCubemapViewMode_2DView),
-		FCanExecuteAction(),
-		FIsActionChecked::CreateSP(this, &FTextureEditorToolkit::HandleCubemapViewModeActionIsChecked, TextureEditorCubemapViewMode_2DView));
-
-	ToolkitCommands->MapAction(
-		Commands.Cubemap3DView,
-		FExecuteAction::CreateSP(this, &FTextureEditorToolkit::HandleCubemapViewModeActionExecute, TextureEditorCubemapViewMode_3DView),
-		FCanExecuteAction(),
-		FIsActionChecked::CreateSP(this, &FTextureEditorToolkit::HandleCubemapViewModeActionIsChecked, TextureEditorCubemapViewMode_3DView));
-	// End - Cubemap Texture Specifics
-
 	ToolkitCommands->MapAction(
 		Commands.TextureBorder,
 		FExecuteAction::CreateSP(this, &FTextureEditorToolkit::HandleTextureBorderActionExecute),
@@ -1703,6 +1710,7 @@ void FTextureEditorToolkit::FillToolbar(FToolBarBuilder& ToolbarBuilder)
 	TSharedRef<SWidget> ExposureControl = MakeExposureContolWidget();
 	TSharedPtr<SWidget> OptionalOpacityControl = IsVolumeTexture() ? TSharedPtr<SWidget>(MakeOpacityControlWidget()) : nullptr;
 	TSharedRef<SWidget> ZoomControl = MakeZoomControlWidget();
+	TSharedRef<SWidget> View3DControl = MakeView3DControlWidget();
 
 	UCurveLinearColorAtlas* Atlas = Cast<UCurveLinearColorAtlas>(GetTexture());
 	if (!Atlas)
@@ -1763,11 +1771,21 @@ void FTextureEditorToolkit::FillToolbar(FToolBarBuilder& ToolbarBuilder)
 			ToolbarBuilder.EndSection();
 		}
 
+		if (IsCubeTexture() || IsVolumeTexture())
+		{
+			ToolbarBuilder.BeginSection("View3D");
+			{
+				ToolbarBuilder.AddWidget(View3DControl);
+			}
+			ToolbarBuilder.EndSection();
+		}
+
 		ToolbarBuilder.BeginSection("Zoom");
 		{
 			ToolbarBuilder.AddWidget(ZoomControl);
 		}
 		ToolbarBuilder.EndSection();
+
 		ToolbarBuilder.BeginSection("Settings");
 		ToolbarBuilder.BeginStyleOverride("CalloutToolbar");
 		{
@@ -1952,7 +1970,7 @@ TSharedRef<SWidget> FTextureEditorToolkit::OnGenerateMipMapLevelMenu()
 TSharedRef<SWidget> FTextureEditorToolkit::OnGenerateSettingsMenu()
 {
 	FMenuBuilder MenuBuilder(true, ToolkitCommands);
-	FTextureEditorViewOptionsMenu::MakeMenu(MenuBuilder, IsVolumeTexture(), IsCubeTexture());
+	FTextureEditorViewOptionsMenu::MakeMenu(MenuBuilder);
 
 	return MenuBuilder.MakeWidget();
 }
@@ -2268,6 +2286,32 @@ void FTextureEditorToolkit::HandleFaceEntryBoxChanged(int32 NewFace)
 	PopulateQuickInfo();
 }
 
+ECheckBoxState FTextureEditorToolkit::HandleView3DCheckBoxIsChecked() const
+{
+	if (IsVolumeTexture())
+	{
+		return GetVolumeViewMode() == TextureEditorVolumeViewMode_VolumeTrace ? ECheckBoxState::Checked : ECheckBoxState::Unchecked;
+	}
+	else
+	{
+		check(IsCubeTexture());
+		return GetCubemapViewMode() == TextureEditorCubemapViewMode_3DView ? ECheckBoxState::Checked : ECheckBoxState::Unchecked;
+	}
+}
+
+void FTextureEditorToolkit::HandleView3DCheckBoxCheckedStateChanged(ECheckBoxState InNewState)
+{
+	if (IsVolumeTexture())
+	{
+		SetVolumeViewMode(InNewState == ECheckBoxState::Checked ? TextureEditorVolumeViewMode_VolumeTrace : TextureEditorVolumeViewMode_DepthSlices);
+	}
+	else
+	{
+		check(IsCubeTexture());
+		SetCubemapViewMode(InNewState == ECheckBoxState::Checked ? TextureEditorCubemapViewMode_3DView : TextureEditorCubemapViewMode_2DView);
+	}	
+}
+
 bool FTextureEditorToolkit::HandleReimportActionCanExecute( ) const
 {
 	if (Texture->IsA<ULightMapTexture2D>() || Texture->IsA<UShadowMapTexture2D>() || Texture->IsA<UTexture2DDynamic>() || Texture->IsA<UTextureRenderTarget>() || Texture->IsA<UCurveLinearColorAtlas>())
@@ -3000,6 +3044,37 @@ TSharedRef<SWidget> FTextureEditorToolkit::MakeZoomControlWidget()
 	return ZoomControl;
 }
 
+TSharedRef<SWidget> FTextureEditorToolkit::MakeView3DControlWidget()
+{
+	// currently both volume textures and cubemaps have only 2 view modes, one of which can be described as "3D View"
+	// ("Trace Into Volume" for volume textures and "3D View" for cubemaps),
+	// therefore view mode for both volume textures and cubemaps can be controlled using the same "3D View" checkbox widget.
+	TSharedRef<SWidget> View3DControl = SNew(SBox)
+		.WidthOverride(80.0f)
+		[
+			SNew(SHorizontalBox)
+			+ SHorizontalBox::Slot()
+			.Padding(4.0f, 0.0f, 4.0f, 0.0f)
+			.VAlign(VAlign_Center)
+			.AutoWidth()
+			[
+				SNew(STextBlock)
+				.Text(NSLOCTEXT("TextureEditor", "3D View", "3D View"))
+			]
+			+ SHorizontalBox::Slot()
+			.Padding(4.0f, 0.0f, 2.0f, 0.0f)
+			.VAlign(VAlign_Center)
+			.AutoWidth()
+			[
+				SNew(SCheckBox)
+				.IsChecked(this, &FTextureEditorToolkit::HandleView3DCheckBoxIsChecked)
+				.OnCheckStateChanged(this, &FTextureEditorToolkit::HandleView3DCheckBoxCheckedStateChanged)
+			]
+		];
+
+	return View3DControl;
+}
+
 void FTextureEditorToolkit::OnEstimateCompressionChanged(ECheckBoxState NewState)
 {
 	OodleCompressedPreviewDDCKey.Set<FString>(FString());
diff --git a/Engine/Source/Editor/TextureEditor/Private/TextureEditorToolkit.h b/Engine/Source/Editor/TextureEditor/Private/TextureEditorToolkit.h
--- a/Engine/Source/Editor/TextureEditor/Private/TextureEditorToolkit.h
+++ b/Engine/Source/Editor/TextureEditor/Private/TextureEditorToolkit.h
@@ -85,8 +85,12 @@ public:
 	virtual void SetVolumeOpacity( float VolumeOpacity ) override;
 	virtual const FRotator& GetVolumeOrientation( ) const override;
 	virtual void SetVolumeOrientation( const FRotator& InOrientation ) override;
+	virtual ETextureEditorVolumeViewMode GetVolumeViewMode() const override;
+	virtual void SetVolumeViewMode(const ETextureEditorVolumeViewMode VolumeViewMode) override;
 	virtual const FRotator& GetCubemapOrientation() const override;
 	virtual void SetCubemapOrientation(const FRotator& InOrientation) override;
+	virtual ETextureEditorCubemapViewMode GetCubemapViewMode() const override;
+	virtual void SetCubemapViewMode(const ETextureEditorCubemapViewMode CubemapViewMode) override;
 	virtual int32 GetExposureBias() const override
 	{
 		return ExposureBias;
@@ -283,6 +287,12 @@ private:
 	// Callback for changing the value of the face index entry box.
 	void HandleFaceEntryBoxChanged(int32 Face);
 
+	// Callback for getting the checked state of the 3D View check box.
+	ECheckBoxState HandleView3DCheckBoxIsChecked() const;
+
+	// Callback for changing the checked state of the 3D View check box.
+	void HandleView3DCheckBoxCheckedStateChanged(ECheckBoxState InNewState);
+
 	// Callback for determining whether the Reimport action can execute.
 	bool HandleReimportActionCanExecute( ) const;
 
@@ -368,6 +378,7 @@ private:
 	TSharedRef<SWidget> MakeExposureContolWidget();
 	TSharedRef<SWidget> MakeOpacityControlWidget();
 	TSharedRef<SWidget> MakeZoomControlWidget();
+	TSharedRef<SWidget> MakeView3DControlWidget();
 private:
 
 	/** The Texture asset being inspected */
@@ -473,9 +484,15 @@ private:
 	// For volume texture, the orientation when tracing.
 	FRotator VolumeOrientation;
 
+	/** This toolkit's current volume view mode **/
+	ETextureEditorVolumeViewMode VolumeViewMode;
+
 	// For cubemap texture, the orientation when previewing the texture as a 3d rendered cube.
 	FRotator CubemapOrientation;
 
+	/** This toolkit's current cubemap view mode **/
+	ETextureEditorCubemapViewMode CubemapViewMode;
+
 	bool bIsVolumeTexture;
 
 	// Objects and callbacks for the custom encoding settings checkbox
diff --git a/Engine/Source/Editor/TextureEditor/Public/Interfaces/ITextureEditorToolkit.h b/Engine/Source/Editor/TextureEditor/Public/Interfaces/ITextureEditorToolkit.h
--- a/Engine/Source/Editor/TextureEditor/Public/Interfaces/ITextureEditorToolkit.h
+++ b/Engine/Source/Editor/TextureEditor/Public/Interfaces/ITextureEditorToolkit.h
@@ -49,8 +49,12 @@ public:
 	virtual void SetVolumeOpacity( float VolumeOpacity ) = 0;
 	virtual const FRotator& GetVolumeOrientation( ) const = 0;
 	virtual void SetVolumeOrientation( const FRotator& InOrientation ) = 0;
+	virtual ETextureEditorVolumeViewMode GetVolumeViewMode() const = 0;
+	virtual void SetVolumeViewMode(const ETextureEditorVolumeViewMode VolumeViewMode) = 0;
 	virtual const FRotator& GetCubemapOrientation() const = 0;
 	virtual void SetCubemapOrientation(const FRotator& InOrientation) = 0;
+	virtual ETextureEditorCubemapViewMode GetCubemapViewMode() const = 0;
+	virtual void SetCubemapViewMode(const ETextureEditorCubemapViewMode CubemapViewMode) = 0;
 	virtual int32 GetExposureBias() const = 0;
 	virtual bool IsVolumeTexture() const = 0;
 public:
