Cleanup cubemap and volume texture preview code.

diff --git a/Engine/Source/Editor/TextureEditor/Private/Models/TextureEditorViewportClient.cpp b/Engine/Source/Editor/TextureEditor/Private/Models/TextureEditorViewportClient.cpp
--- a/Engine/Source/Editor/TextureEditor/Private/Models/TextureEditorViewportClient.cpp
+++ b/Engine/Source/Editor/TextureEditor/Private/Models/TextureEditorViewportClient.cpp
@@ -119,7 +119,7 @@ void FTextureEditorViewportClient::Draw(FViewport* Viewport, FCanvas* Canvas)
 			// otherwise generate a view matrix which displays the selected face using DDS orientation and face order
 			// (positive x, negative x, positive y, negative y, positive z, negative z)
 			const FMatrix44f& ViewMatrix =
-				FaceIndex < 0 ? FMatrix44f(FRotationMatrix::Make(TextureEditorPinned->GetCubemapOrientation())) :
+				FaceIndex < 0 ? FMatrix44f(FRotationMatrix::Make(TextureEditorPinned->GetOrientation())) :
 				FaceIndex == 0 ? FMatrix44f(FVector3f::ForwardVector, FVector3f::DownVector, FVector3f::RightVector, IdentityW) :
 				FaceIndex == 1 ? FMatrix44f(FVector3f::BackwardVector, FVector3f::UpVector, FVector3f::RightVector, IdentityW) :
 				FaceIndex == 2 ? FMatrix44f(FVector3f::RightVector, FVector3f::ForwardVector, FVector3f::DownVector, IdentityW) :
@@ -137,7 +137,7 @@ void FTextureEditorViewportClient::Draw(FViewport* Viewport, FCanvas* Canvas)
 				MipLevel, 
 				(float)TextureEditorPinned->GetVolumeOpacity(),
 				true, 
-				TextureEditorPinned->GetVolumeOrientation());
+				TextureEditorPinned->GetOrientation());
 		}
 		else if (RTTextureVolume)
 		{
@@ -147,7 +147,7 @@ void FTextureEditorViewportClient::Draw(FViewport* Viewport, FCanvas* Canvas)
 				MipLevel,
 				(float)TextureEditorPinned->GetVolumeOpacity(),
 				true,
-				TextureEditorPinned->GetVolumeOrientation());
+				TextureEditorPinned->GetOrientation());
 		}
 		else if (Texture2D)
 		{
@@ -287,59 +287,38 @@ bool FTextureEditorViewportClient::InputKey(const FInputKeyEventArgs& InEventArg
 
 			return true;
 		}
-		else if (InEventArgs.Key == EKeys::RightMouseButton)
+		else if (InEventArgs.Key == EKeys::MiddleMouseButton && TextureEditorPtr.Pin()->IsUsingOrientation())
 		{
-			TextureEditorPtr.Pin()->SetVolumeOrientation(FRotator(90, 0, -90));
+			TextureEditorPtr.Pin()->ResetOrientation();
 		}
 	}
 	return false;
 }
 
-bool IsTextureUsingVolumeOrientation(UTexture* Texture)
-{
-	return Texture && (Cast<UVolumeTexture>(Texture) || Cast<UTextureRenderTargetVolume>(Texture));
-}
-
 bool FTextureEditorViewportClient::InputAxis(FViewport* Viewport, FInputDeviceId DeviceId, FKey Key, float Delta, float DeltaTime, int32 NumSamples, bool bGamepad)
 {
 	if (Key == EKeys::MouseX || Key == EKeys::MouseY)
 	{
-		UTexture* Texture = TextureEditorPtr.Pin()->GetTexture();
-		if (IsTextureUsingVolumeOrientation(Texture))
-		{
-			FRotator DeltaRotator(ForceInitToZero);
-			const float RotationSpeed = .2f;
-			if (Key == EKeys::MouseY)
-			{
-				DeltaRotator.Pitch = Delta * RotationSpeed;
-			}
-			else
-			{
-				DeltaRotator.Yaw = Delta * RotationSpeed;
-			}
-
-			TextureEditorPtr.Pin()->SetVolumeOrientation((FRotationMatrix::Make(DeltaRotator) * FRotationMatrix::Make(TextureEditorPtr.Pin()->GetVolumeOrientation())).Rotator());
-		}
-		else if (IsTextureUsingCubemapOrientation() && Viewport->KeyState(EKeys::LeftMouseButton))
+		TSharedPtr<ITextureEditorToolkit> TextureEditorPinned = TextureEditorPtr.Pin();
+		if (TextureEditorPinned->IsUsingOrientation() && Viewport->KeyState(EKeys::LeftMouseButton))
 		{
-			FRotator Orienation = TextureEditorPtr.Pin()->GetCubemapOrientation();
 			const float RotationSpeed = .2f;
-			if (Key == EKeys::MouseY)
+			FRotator DeltaOrientation = FRotator(Key == EKeys::MouseY ? Delta * RotationSpeed : 0, Key == EKeys::MouseX ? Delta * RotationSpeed : 0, 0);
+			if (TextureEditorPinned->IsVolumeTexture())
 			{
-				Orienation.Pitch += Delta * RotationSpeed;
+				TextureEditorPinned->SetOrientation((FRotationMatrix::Make(DeltaOrientation) * FRotationMatrix::Make(TextureEditorPinned->GetOrientation())).Rotator());
 			}
 			else
 			{
-				Orienation.Yaw += Delta * RotationSpeed;
+				TextureEditorPinned->SetOrientation(TextureEditorPinned->GetOrientation() + DeltaOrientation);
 			}
-			TextureEditorPtr.Pin()->SetCubemapOrientation(Orienation);
 		}
 		else if (ShouldUseMousePanning(Viewport))
 		{
 			TSharedPtr<STextureEditorViewport> EditorViewport = TextureEditorViewportPtr.Pin();
 
 			int32 Width, Height, Depth, ArraySize;
-			TextureEditorPtr.Pin()->CalculateTextureDimensions(Width, Height, Depth, ArraySize, true);
+			TextureEditorPinned->CalculateTextureDimensions(Width, Height, Depth, ArraySize, true);
 
 			if (Key == EKeys::MouseY)
 			{
@@ -362,17 +341,9 @@ bool FTextureEditorViewportClient::InputAxis(FViewport* Viewport, FInputDeviceId
 	return false;
 }
 
-bool FTextureEditorViewportClient::IsTextureUsingCubemapOrientation() const
-{
-	TSharedPtr<ITextureEditorToolkit> TextureEditorPinned = TextureEditorPtr.Pin();
-	UTexture* Texture = TextureEditorPinned->GetTexture();
-	return Texture && (Texture->IsA(UTextureCube::StaticClass()) || Texture->IsA(UTextureRenderTargetCube::StaticClass()) || Texture->IsA(UTextureCubeArray::StaticClass())) &&
-		TextureEditorPinned->GetFace() < 0 && TextureEditorPinned->GetCubemapViewMode() == TextureEditorCubemapViewMode_3DView;
-}
-
 bool FTextureEditorViewportClient::ShouldUseMousePanning(FViewport* Viewport) const
 {
-	if (!IsTextureUsingVolumeOrientation(TextureEditorPtr.Pin()->GetTexture()) && Viewport->KeyState(EKeys::RightMouseButton))
+	if (Viewport->KeyState(EKeys::RightMouseButton))
 	{
 		TSharedPtr<STextureEditorViewport> EditorViewport = TextureEditorViewportPtr.Pin();
 		return EditorViewport.IsValid() && EditorViewport->GetVerticalScrollBar().IsValid() && EditorViewport->GetHorizontalScrollBar().IsValid();
diff --git a/Engine/Source/Editor/TextureEditor/Private/Models/TextureEditorViewportClient.h b/Engine/Source/Editor/TextureEditor/Private/Models/TextureEditorViewportClient.h
--- a/Engine/Source/Editor/TextureEditor/Private/Models/TextureEditorViewportClient.h
+++ b/Engine/Source/Editor/TextureEditor/Private/Models/TextureEditorViewportClient.h
@@ -59,9 +59,6 @@ private:
 	/** TRUE if right clicking and dragging for panning a texture 2D */
 	bool ShouldUseMousePanning(FViewport* Viewport) const;
 
-	/** Returns true if the texture is using cubemap orientation */
-	bool IsTextureUsingCubemapOrientation() const;
-
 private:
 	/** Pointer back to the Texture editor tool that owns us */
 	TWeakPtr<ITextureEditorToolkit> TextureEditorPtr;
diff --git a/Engine/Source/Editor/TextureEditor/Private/TextureEditorToolkit.cpp b/Engine/Source/Editor/TextureEditor/Private/TextureEditorToolkit.cpp
--- a/Engine/Source/Editor/TextureEditor/Private/TextureEditorToolkit.cpp
+++ b/Engine/Source/Editor/TextureEditor/Private/TextureEditorToolkit.cpp
@@ -121,8 +121,6 @@ void FTextureEditorToolkit::PostTextureRecode()
 FTextureEditorToolkit::FTextureEditorToolkit()
 	: Texture(nullptr)
 	, VolumeOpacity(1.f)
-	, VolumeOrientation(90, 0, -90)
-	, CubemapOrientation(0, 0, 0)
 {
 }
 
@@ -260,6 +258,7 @@ void FTextureEditorToolkit::InitTextureEditor( const EToolkitMode::Type Mode, co
 	VolumeViewMode = Settings.VolumeViewMode;
 	CubemapViewMode = Settings.CubemapViewMode;
 
+	ResetOrientation();
 	Zoom = 1.0f;
 
 	// Register our commands. This will only register them if not previously registered
@@ -894,16 +893,6 @@ void FTextureEditorToolkit::SetVolumeOpacity(float InVolumeOpacity)
 	VolumeOpacity = FMath::Clamp(InVolumeOpacity, 0.f, 1.f);
 }
 
-const FRotator& FTextureEditorToolkit::GetVolumeOrientation() const
-{
-	return VolumeOrientation;
-}
-
-void FTextureEditorToolkit::SetVolumeOrientation(const FRotator& InOrientation)
-{
-	VolumeOrientation = InOrientation;
-}
-
 ETextureEditorVolumeViewMode FTextureEditorToolkit::GetVolumeViewMode() const
 {
 	// Each texture editor keeps a local volume view mode so that it can be changed without affecting other open editors
@@ -921,16 +910,6 @@ void FTextureEditorToolkit::SetVolumeViewMode(const ETextureEditorVolumeViewMode
 	Settings.PostEditChange();
 }
 
-const FRotator& FTextureEditorToolkit::GetCubemapOrientation() const
-{
-	return CubemapOrientation;
-}
-
-void FTextureEditorToolkit::SetCubemapOrientation(const FRotator& InOrientation)
-{
-	CubemapOrientation = InOrientation;
-}
-
 ETextureEditorCubemapViewMode FTextureEditorToolkit::GetCubemapViewMode() const
 {
 	// Each texture editor keeps a local cubemap view mode so that it can be changed without affecting other open editors
@@ -948,6 +927,27 @@ void FTextureEditorToolkit::SetCubemapViewMode(const ETextureEditorCubemapViewMo
 	Settings.PostEditChange();
 }
 
+bool FTextureEditorToolkit::IsUsingOrientation() const
+{
+	return (IsVolumeTexture() && GetVolumeViewMode() == TextureEditorVolumeViewMode_VolumeTrace) ||
+		(IsCubeTexture() && GetCubemapViewMode() == TextureEditorCubemapViewMode_3DView && GetFace() < 0);
+}
+
+const FRotator& FTextureEditorToolkit::GetOrientation() const
+{
+	return Orientation;
+}
+
+void FTextureEditorToolkit::SetOrientation(const FRotator& InOrientation)
+{
+	Orientation = InOrientation;
+}
+
+void FTextureEditorToolkit::ResetOrientation()
+{
+	SetOrientation(IsVolumeTexture() ? FRotator(90, 0, -90) : FRotator(0, 0, 0));
+}
+
 /* IToolkit interface
  *****************************************************************************/
 
diff --git a/Engine/Source/Editor/TextureEditor/Private/TextureEditorToolkit.h b/Engine/Source/Editor/TextureEditor/Private/TextureEditorToolkit.h
--- a/Engine/Source/Editor/TextureEditor/Private/TextureEditorToolkit.h
+++ b/Engine/Source/Editor/TextureEditor/Private/TextureEditorToolkit.h
@@ -83,14 +83,14 @@ public:
 	virtual void ZoomOut( ) override;
 	virtual float GetVolumeOpacity( ) const override;
 	virtual void SetVolumeOpacity( float VolumeOpacity ) override;
-	virtual const FRotator& GetVolumeOrientation( ) const override;
-	virtual void SetVolumeOrientation( const FRotator& InOrientation ) override;
 	virtual ETextureEditorVolumeViewMode GetVolumeViewMode() const override;
 	virtual void SetVolumeViewMode(const ETextureEditorVolumeViewMode VolumeViewMode) override;
-	virtual const FRotator& GetCubemapOrientation() const override;
-	virtual void SetCubemapOrientation(const FRotator& InOrientation) override;
 	virtual ETextureEditorCubemapViewMode GetCubemapViewMode() const override;
 	virtual void SetCubemapViewMode(const ETextureEditorCubemapViewMode CubemapViewMode) override;
+	virtual bool IsUsingOrientation() const override;
+	virtual const FRotator& GetOrientation() const override;
+	virtual void SetOrientation(const FRotator& InOrientation) override;
+	virtual void ResetOrientation() override;
 	virtual int32 GetExposureBias() const override
 	{
 		return ExposureBias;
@@ -481,18 +481,15 @@ private:
 	// For volume texture, defines an opacity to see through the volume when tracing.
 	float VolumeOpacity;
 
-	// For volume texture, the orientation when tracing.
-	FRotator VolumeOrientation;
-
 	/** This toolkit's current volume view mode **/
 	ETextureEditorVolumeViewMode VolumeViewMode;
 
-	// For cubemap texture, the orientation when previewing the texture as a 3d rendered cube.
-	FRotator CubemapOrientation;
-
 	/** This toolkit's current cubemap view mode **/
 	ETextureEditorCubemapViewMode CubemapViewMode;
 
+	// Orientation of the texture in 3d when using "3D View" mode
+	FRotator Orientation;
+
 	bool bIsVolumeTexture;
 
 	// Objects and callbacks for the custom encoding settings checkbox
diff --git a/Engine/Source/Editor/TextureEditor/Public/Interfaces/ITextureEditorToolkit.h b/Engine/Source/Editor/TextureEditor/Public/Interfaces/ITextureEditorToolkit.h
--- a/Engine/Source/Editor/TextureEditor/Public/Interfaces/ITextureEditorToolkit.h
+++ b/Engine/Source/Editor/TextureEditor/Public/Interfaces/ITextureEditorToolkit.h
@@ -47,14 +47,14 @@ public:
 	virtual double CalculateDisplayedZoomLevel() const = 0;
 	virtual float GetVolumeOpacity( ) const = 0;
 	virtual void SetVolumeOpacity( float VolumeOpacity ) = 0;
-	virtual const FRotator& GetVolumeOrientation( ) const = 0;
-	virtual void SetVolumeOrientation( const FRotator& InOrientation ) = 0;
 	virtual ETextureEditorVolumeViewMode GetVolumeViewMode() const = 0;
 	virtual void SetVolumeViewMode(const ETextureEditorVolumeViewMode VolumeViewMode) = 0;
-	virtual const FRotator& GetCubemapOrientation() const = 0;
-	virtual void SetCubemapOrientation(const FRotator& InOrientation) = 0;
 	virtual ETextureEditorCubemapViewMode GetCubemapViewMode() const = 0;
 	virtual void SetCubemapViewMode(const ETextureEditorCubemapViewMode CubemapViewMode) = 0;
+	virtual bool IsUsingOrientation() const = 0;
+	virtual const FRotator& GetOrientation() const = 0;
+	virtual void SetOrientation(const FRotator& InOrientation) = 0;
+	virtual void ResetOrientation() = 0;
 	virtual int32 GetExposureBias() const = 0;
 	virtual bool IsVolumeTexture() const = 0;
 public:
