Display proper stats for cubemaps in the texture editor.

diff --git a/Engine/Source/Editor/TextureEditor/Private/TextureEditorToolkit.cpp b/Engine/Source/Editor/TextureEditor/Private/TextureEditorToolkit.cpp
--- a/Engine/Source/Editor/TextureEditor/Private/TextureEditorToolkit.cpp
+++ b/Engine/Source/Editor/TextureEditor/Private/TextureEditorToolkit.cpp
@@ -725,31 +725,37 @@ void FTextureEditorToolkit::PopulateQuickInfo( )
 	// Texture asset size
 	const int64 ResourceSize = PlatformDataPtr && *PlatformDataPtr ? (*PlatformDataPtr)->GetPayloadSize(ActualMipBias) : Texture->GetResourceSizeBytes(EResourceSizeMode::Exclusive);
 
-	FNumberFormattingOptions SizeOptions;
-	SizeOptions.UseGrouping = false;
-	SizeOptions.MaximumFractionalDigits = 0;
-
-	// Cubes are previewed as longlat unwrapped 2D textures in 2D view mode when face index is not specified.
-	// These have 2x the width of a cube face.
-	if (IsCubeTexture() && GetCubemapViewMode() == TextureEditorCubemapViewMode_2DView && GetFace() < 0)
-	{
-		PreviewEffectiveTextureWidth *= 2;
-	}
-
-	FNumberFormattingOptions Options;
-	Options.UseGrouping = false;
-
-	FText CubemapAdd;
+	FText ImportedCubemapInfo;
+	FText DisplayedCubemapInfo;
+	FText InGameCubemapInfo;
 	if (bIsCube)
 	{
-		CubemapAdd = NSLOCTEXT("TextureEditor", "QuickInfo_PerCubeSide", "*6 (CubeMap)");
+		InGameCubemapInfo = NSLOCTEXT("TextureEditor", "QuickInfo_PerCubeSide", "*6 (Cubemap)");
+		ImportedCubemapInfo = Texture->Source.IsLongLatCubemap() ? NSLOCTEXT("TextureEditor", "QuickInfo_LongLat", " (LongLat)") : InGameCubemapInfo;
+		if (GetFace() >= 0)
+		{
+			DisplayedCubemapInfo = NSLOCTEXT("TextureEditor", "QuickInfo_CubemapFace", " (Cubemap face)");
+		}
+		else if (GetCubemapViewMode() == TextureEditorCubemapViewMode_2DView)
+		{
+			// when using 2D view mode and face index is not specified, cubemaps are previewed as longlat unwraps, which have double width of a cube face
+			PreviewEffectiveTextureWidth *= 2;
+			DisplayedCubemapInfo = NSLOCTEXT("TextureEditor", "QuickInfo_LongLat", " (LongLat)");
+		}
+		else 
+		{
+			DisplayedCubemapInfo = InGameCubemapInfo;
+		}
 	}
 
+	FNumberFormattingOptions FormatOptions;
+	FormatOptions.UseGrouping = false;
+
 	if (bIsVolume)
 	{
-		ImportedText->SetText(FText::Format( NSLOCTEXT("TextureEditor", "QuickInfo_Imported_3x", "Imported: {0}x{1}x{2}"), FText::AsNumber(ImportedWidth, &Options), FText::AsNumber(ImportedHeight, &Options), FText::AsNumber(ImportedDepth, &Options)));
-		CurrentText->SetText(FText::Format( NSLOCTEXT("TextureEditor", "QuickInfo_Displayed_3x", "Displayed: {0}x{1}x{2}"), FText::AsNumber(PreviewEffectiveTextureWidth, &Options ), FText::AsNumber(PreviewEffectiveTextureHeight, &Options), FText::AsNumber(PreviewEffectiveTextureDepth, &Options)));
-		MaxInGameText->SetText(FText::Format( NSLOCTEXT("TextureEditor", "QuickInfo_MaxInGame_3x_v1", "Max In-Game: {0}x{1}x{2}"), FText::AsNumber(MaxInGameWidth, &Options), FText::AsNumber(MaxInGameHeight, &Options), FText::AsNumber(MaxInGameDepth, &Options)));
+		ImportedText->SetText(FText::Format( NSLOCTEXT("TextureEditor", "QuickInfo_Imported_3x", "Imported: {0}x{1}x{2}"), FText::AsNumber(ImportedWidth, &FormatOptions), FText::AsNumber(ImportedHeight, &FormatOptions), FText::AsNumber(ImportedDepth, &FormatOptions)));
+		CurrentText->SetText(FText::Format( NSLOCTEXT("TextureEditor", "QuickInfo_Displayed_3x", "Displayed: {0}x{1}x{2}"), FText::AsNumber(PreviewEffectiveTextureWidth, &FormatOptions ), FText::AsNumber(PreviewEffectiveTextureHeight, &FormatOptions), FText::AsNumber(PreviewEffectiveTextureDepth, &FormatOptions)));
+		MaxInGameText->SetText(FText::Format( NSLOCTEXT("TextureEditor", "QuickInfo_MaxInGame_3x_v1", "Max In-Game: {0}x{1}x{2}"), FText::AsNumber(MaxInGameWidth, &FormatOptions), FText::AsNumber(MaxInGameHeight, &FormatOptions), FText::AsNumber(MaxInGameDepth, &FormatOptions)));
 
 		if (GetVolumeViewMode() == ETextureEditorVolumeViewMode::TextureEditorVolumeViewMode_VolumeTrace)
 		{
@@ -766,18 +772,18 @@ void FTextureEditorToolkit::PopulateQuickInfo( )
 	}
 	else if (bIsArray)
 	{
-		ImportedText->SetText(FText::Format(NSLOCTEXT("TextureEditor", "QuickInfo_Imported_3x_v2", "Imported: {0}x{1}*{2}"), FText::AsNumber(ImportedWidth, &Options), FText::AsNumber(ImportedHeight, &Options), FText::AsNumber(ArraySize, &Options)));
-		CurrentText->SetText(FText::Format(NSLOCTEXT("TextureEditor", "QuickInfo_Displayed_3x_v2", "Displayed: {0}x{1}{2}*{3}"), FText::AsNumber(PreviewEffectiveTextureWidth, &Options), FText::AsNumber(PreviewEffectiveTextureHeight, &Options), CubemapAdd, FText::AsNumber(ArraySize, &Options)));
-		MaxInGameText->SetText(FText::Format(NSLOCTEXT("TextureEditor", "QuickInfo_MaxInGame_3x_v2", "Max In-Game: {0}x{1}{2}*{3}"), FText::AsNumber(MaxInGameWidth, &Options), FText::AsNumber(MaxInGameHeight, &Options), CubemapAdd, FText::AsNumber(ArraySize, &Options)));
+		ImportedText->SetText(FText::Format(NSLOCTEXT("TextureEditor", "QuickInfo_Imported_3x_v2", "Imported: {0}x{1}{2}*{3}"), FText::AsNumber(ImportedWidth, &FormatOptions), FText::AsNumber(ImportedHeight, &FormatOptions), ImportedCubemapInfo, FText::AsNumber(ArraySize, &FormatOptions)));
+		CurrentText->SetText(FText::Format(NSLOCTEXT("TextureEditor", "QuickInfo_Displayed_3x_v2", "Displayed: {0}x{1}{2}*{3}"), FText::AsNumber(PreviewEffectiveTextureWidth, &FormatOptions), FText::AsNumber(PreviewEffectiveTextureHeight, &FormatOptions), DisplayedCubemapInfo, FText::AsNumber(ArraySize, &FormatOptions)));
+		MaxInGameText->SetText(FText::Format(NSLOCTEXT("TextureEditor", "QuickInfo_MaxInGame_3x_v2", "Max In-Game: {0}x{1}{2}*{3}"), FText::AsNumber(MaxInGameWidth, &FormatOptions), FText::AsNumber(MaxInGameHeight, &FormatOptions), InGameCubemapInfo, FText::AsNumber(ArraySize, &FormatOptions)));
 	}
 	else
 	{
-	    ImportedText->SetText(FText::Format( NSLOCTEXT("TextureEditor", "QuickInfo_Imported_2x", "Imported: {0}x{1}"), FText::AsNumber(ImportedWidth, &Options), FText::AsNumber(ImportedHeight, &Options)));
-		CurrentText->SetText(FText::Format( NSLOCTEXT("TextureEditor", "QuickInfo_Displayed_2x", "Displayed: {0}x{1}{2}"), FText::AsNumber(PreviewEffectiveTextureWidth, &Options ), FText::AsNumber(PreviewEffectiveTextureHeight, &Options), CubemapAdd));
-		MaxInGameText->SetText(FText::Format( NSLOCTEXT("TextureEditor", "QuickInfo_MaxInGame_2x", "Max In-Game: {0}x{1}{2}"), FText::AsNumber(MaxInGameWidth, &Options), FText::AsNumber(MaxInGameHeight, &Options), CubemapAdd));
+	    ImportedText->SetText(FText::Format( NSLOCTEXT("TextureEditor", "QuickInfo_Imported_2x", "Imported: {0}x{1}{2}"), FText::AsNumber(ImportedWidth, &FormatOptions), FText::AsNumber(ImportedHeight, &FormatOptions), ImportedCubemapInfo));
+		CurrentText->SetText(FText::Format( NSLOCTEXT("TextureEditor", "QuickInfo_Displayed_2x", "Displayed: {0}x{1}{2}"), FText::AsNumber(PreviewEffectiveTextureWidth, &FormatOptions ), FText::AsNumber(PreviewEffectiveTextureHeight, &FormatOptions), DisplayedCubemapInfo));
+		MaxInGameText->SetText(FText::Format( NSLOCTEXT("TextureEditor", "QuickInfo_MaxInGame_2x", "Max In-Game: {0}x{1}{2}"), FText::AsNumber(MaxInGameWidth, &FormatOptions), FText::AsNumber(MaxInGameHeight, &FormatOptions), InGameCubemapInfo));
 	}
 
-	SizeText->SetText(FText::Format(NSLOCTEXT("TextureEditor", "QuickInfo_ResourceSize", "Resource Size: {0} KB"), FText::AsNumber(FMath::DivideAndRoundNearest(ResourceSize, (int64)1024), &SizeOptions)));
+	SizeText->SetText(FText::Format(NSLOCTEXT("TextureEditor", "QuickInfo_ResourceSize", "Resource Size: {0} KB"), FText::AsNumber(FMath::DivideAndRoundNearest(ResourceSize, (int64)1024), &FormatOptions)));
 
 	FText Method = Texture->IsCurrentlyVirtualTextured() ? NSLOCTEXT("TextureEditor", "QuickInfo_MethodVirtualStreamed", "Virtual Streamed")
 													: (!Texture->IsStreamable() ? NSLOCTEXT("TextureEditor", "QuickInfo_MethodNotStreamed", "Not Streamed") 
