Display correct imported size for longlat cubemaps smaller than 32 pixels.

diff --git a/Engine/Source/Editor/TextureEditor/Private/TextureEditorToolkit.cpp b/Engine/Source/Editor/TextureEditor/Private/TextureEditorToolkit.cpp
--- a/Engine/Source/Editor/TextureEditor/Private/TextureEditorToolkit.cpp
+++ b/Engine/Source/Editor/TextureEditor/Private/TextureEditorToolkit.cpp
@@ -704,9 +704,9 @@ void FTextureEditorToolkit::PopulateQuickInfo( )
 	const int32 NumSurfaces = Texture->GetSurfaceArraySize();
 	const int32 ArraySize = bIsArray ? (bIsCube ? NumSurfaces / 6 : NumSurfaces) : 0;
 
-	const int32 ImportedWidth = FMath::Max(SurfaceWidth, Texture->Source.GetSizeX());
-	const int32 ImportedHeight =  FMath::Max(SurfaceHeight, Texture->Source.GetSizeY());
-	const int32 ImportedDepth = FMath::Max(SurfaceDepth, bIsVolume ? Texture->Source.GetNumSlices() : 0);
+	const int32 ImportedWidth = Texture->Source.IsValid() ? Texture->Source.GetSizeX() : SurfaceWidth;
+	const int32 ImportedHeight = Texture->Source.IsValid() ? Texture->Source.GetSizeY() : SurfaceHeight;
+	const int32 ImportedDepth = Texture->Source.IsValid() ? (bIsVolume ? Texture->Source.GetNumSlices() : 0) : SurfaceDepth;
 
 	const FStreamableRenderResourceState SRRState = Texture->GetStreamableResourceState();
 	const int32 ActualMipBias = SRRState.IsValid() ? (SRRState.ResidentFirstLODIdx() + SRRState.AssetLODBias) : Texture->GetCachedLODBias();
