Scene capture fix for render targets not displaying in Texture Editor, by starting out with the alpha channel disabled for render targets.  Scene capture render targets have a depth mask in the alpha channel, which causes them to appear as fully transparent (checkerboard), and it isn't easily discoverable that toggling off the alpha channel would solve the issue, so it's helpful to start out with it disabled.

diff --git a/Engine/Source/Editor/TextureEditor/Private/TextureEditorToolkit.cpp b/Engine/Source/Editor/TextureEditor/Private/TextureEditorToolkit.cpp
--- a/Engine/Source/Editor/TextureEditor/Private/TextureEditorToolkit.cpp
+++ b/Engine/Source/Editor/TextureEditor/Private/TextureEditorToolkit.cpp
@@ -234,6 +234,15 @@ void FTextureEditorToolkit::InitTextureEditor( const EToolkitMode::Type Mode, co
 		break;
 	}
 
+	// Start out with alpha channel unchecked for render targets.  Scene capture render targets have a depth mask
+	// in the alpha channel, which causes them to appear as fully transparent (checkerboard), and it isn't easily
+	// discoverable that toggling off the alpha channel would solve the issue, so it's helpful to start out with
+	// it disabled.
+	if (Texture->GetTextureClass() == ETextureClass::RenderTarget)
+	{
+		bIsAlphaChannel = false;
+	}
+
 	bIsDesaturation = false;
 
 	PreviewEffectiveTextureWidth = 0;
