stop reimport setting "defer compression"
causes unexpected visual quality change when you reimport

diff --git a/Engine/Source/Editor/TextureEditor/Private/TextureEditorToolkit.cpp b/Engine/Source/Editor/TextureEditor/Private/TextureEditorToolkit.cpp
--- a/Engine/Source/Editor/TextureEditor/Private/TextureEditorToolkit.cpp
+++ b/Engine/Source/Editor/TextureEditor/Private/TextureEditorToolkit.cpp
@@ -260,8 +260,6 @@ void FTextureEditorToolkit::InitTextureEditor( const EToolkitMode::Type Mode, co
 	SpecifiedFace = 0;
 	bUseSpecifiedFace = false;
 
-	SavedCompressionSetting = false;
-
 	// Start at whatever the last used zoom mode, volume view mode and cubemap view mode were
 	const UTextureEditorSettings& Settings = *GetDefault<UTextureEditorSettings>();
 	ZoomMode = Settings.ZoomMode;
@@ -2396,12 +2394,6 @@ void FTextureEditorToolkit::HandleReimportManagerPostReimport( UObject* InObject
 		return;
 	}
 
-	if (!bSuccess)
-	{
-		// Failed, restore the compression flag
-		Texture->DeferCompression = SavedCompressionSetting;
-	}
-
 	// Re-enable viewport rendering now that the texture should be in a known state again
 	TextureViewport->EnableRendering();
 }
@@ -2415,10 +2407,6 @@ void FTextureEditorToolkit::HandleReimportManagerPreReimport( UObject* InObject
 		return;
 	}
 
-	// Prevent the texture from being compressed immediately, so the user can see the results
-	SavedCompressionSetting = Texture->DeferCompression;
-	Texture->DeferCompression = true;
-
 	// Disable viewport rendering until the texture has finished re-importing
 	TextureViewport->DisableRendering();
 }
diff --git a/Engine/Source/Editor/TextureEditor/Private/TextureEditorToolkit.h b/Engine/Source/Editor/TextureEditor/Private/TextureEditorToolkit.h
--- a/Engine/Source/Editor/TextureEditor/Private/TextureEditorToolkit.h
+++ b/Engine/Source/Editor/TextureEditor/Private/TextureEditorToolkit.h
@@ -468,9 +468,6 @@ private:
 	/** When true, the specified cubemap face index value is used, otherwise the value of -1 is used.*/
 	bool bUseSpecifiedFace;
 
-	/** During re-import, cache this setting so it can be restored if necessary */
-	bool SavedCompressionSetting;
-
 	/** The texture's zoom factor. */
 	double Zoom;
 
