Correcting Oodle nomenclature

diff --git a/Engine/Source/Developer/DeveloperToolSettings/Classes/Settings/ProjectPackagingSettings.h b/Engine/Source/Developer/DeveloperToolSettings/Classes/Settings/ProjectPackagingSettings.h
--- a/Engine/Source/Developer/DeveloperToolSettings/Classes/Settings/ProjectPackagingSettings.h
+++ b/Engine/Source/Developer/DeveloperToolSettings/Classes/Settings/ProjectPackagingSettings.h
@@ -368,16 +368,16 @@ public:
 	FString PackageCompressionMethod;
 		
 	/*
-	 * For compressors with variable levels, select the encoder effort level, which makes packages smaller but takes more time to encode.
+	 * For compressors with variable levels, select the compressor effort level, which makes packages smaller but takes more time to encode.
 	 * This does not affect decode speed.  For faster iteration, use lower effort levels (eg. 1)
 	 */
-	UPROPERTY(config, EditAnywhere, Category = Packaging, AdvancedDisplay, meta = (DisplayName = "Encoder Effort Level for Debug & Development"))
+	UPROPERTY(config, EditAnywhere, Category = Packaging, AdvancedDisplay, meta = (DisplayName = "Compressor Effort Level for Debug & Development"))
 	int32 PackageCompressionLevel_DebugDevelopment;
 	
-	UPROPERTY(config, EditAnywhere, Category = Packaging, AdvancedDisplay, meta = (DisplayName = "Encoder Effort Level for Test & Shipping"))
+	UPROPERTY(config, EditAnywhere, Category = Packaging, AdvancedDisplay, meta = (DisplayName = "Compressor Effort Level for Test & Shipping"))
 	int32 PackageCompressionLevel_TestShipping;
 	
-	UPROPERTY(config, EditAnywhere, Category = Packaging, AdvancedDisplay, meta = (DisplayName = "Encoder Effort Level for Distribution"))
+	UPROPERTY(config, EditAnywhere, Category = Packaging, AdvancedDisplay, meta = (DisplayName = "Compressor Effort Level for Distribution"))
 	int32 PackageCompressionLevel_Distribution;
 
 	/**
diff --git a/Engine/Source/Editor/TextureEditor/Private/TextureEditorToolkit.cpp b/Engine/Source/Editor/TextureEditor/Private/TextureEditorToolkit.cpp
--- a/Engine/Source/Editor/TextureEditor/Private/TextureEditorToolkit.cpp
+++ b/Engine/Source/Editor/TextureEditor/Private/TextureEditorToolkit.cpp
@@ -1153,7 +1153,7 @@ void FTextureEditorToolkit::CreateInternalWidgets()
 	// Default to Distribution
 	TSharedPtr<FString> InitialPackagingSetting = PackagingSettingsNames[2];
 
-	// Determine which oodle encoder they are using.
+	// Determine which oodle compressor they are using.
 	const TCHAR* CompressorName = 0;
 	{
 		// Validity check the string by trying to convert to enum.
@@ -1367,7 +1367,7 @@ void FTextureEditorToolkit::CreateInternalWidgets()
 						[
 							SNew(STextBlock)
 							.Text(LOCTEXT("OodleTab_Label_OverrideCompression", "Enabled:"))
-							.ToolTipText(LOCTEXT("OodleTab_ToolTip_OverrideCompression", "If checked, allows you to experiment with Oodle RDO compression settings to visualize results."))
+							.ToolTipText(LOCTEXT("OodleTab_ToolTip_OverrideEncoding", "If checked, allows you to experiment with Oodle RDO encoder settings to visualize results."))
 						]
 					+ SVerticalBox::Slot()
 						.AutoHeight()
@@ -1493,7 +1493,7 @@ void FTextureEditorToolkit::CreateInternalWidgets()
 						[
 							SNew(STextBlock)
 							.Text(LOCTEXT("OodleTab_Label_EncoderSettings", "Packaging Configuration:"))
-							.ToolTipText(LOCTEXT("OodleTab_ToolTip_EncoderSettings", "Which packaging configuration to pull from for determining which Oodle encoder and compression level to use."))
+							.ToolTipText(LOCTEXT("OodleTab_ToolTip_CompressorSettings", "Which packaging configuration to pull from for determining which Oodle compressor and compression level to use."))
 						]
 					+ SVerticalBox::Slot()
 						.AutoHeight()
@@ -1501,8 +1501,8 @@ void FTextureEditorToolkit::CreateInternalWidgets()
 						.Padding(6)
 						[
 							SNew(STextBlock)
-							.Text(LOCTEXT("OodleTab_Label_EstimateEncoder", "Oodle Encoder:"))
-							.ToolTipText(LOCTEXT("OodleTab_ToolTip_EstimateEncoder", "The oodle encoder to use for estimating. Pulled from the packaging configuration specified above."))
+							.Text(LOCTEXT("OodleTab_Label_EstimateCompressor", "Oodle Compressor:"))
+							.ToolTipText(LOCTEXT("OodleTab_ToolTip_EstimateCompressor", "The oodle compressor to use for estimating. Pulled from the packaging configuration specified above."))
 						]
 					+ SVerticalBox::Slot()
 						.AutoHeight()
@@ -1574,7 +1574,7 @@ void FTextureEditorToolkit::CreateInternalWidgets()
 						.VAlign(VAlign_Center)
 						.Padding(8)
 						[
-							SAssignNew(OodleEncoderUsed, STextBlock)
+							SAssignNew(OodleCompressorUsed, STextBlock)
 							.Text(FText::AsCultureInvariant(CompressorName))
 							.IsEnabled(this, &FTextureEditorToolkit::EstimateCompressionEnabled)
 						]
diff --git a/Engine/Source/Editor/TextureEditor/Private/TextureEditorToolkit.h b/Engine/Source/Editor/TextureEditor/Private/TextureEditorToolkit.h
--- a/Engine/Source/Editor/TextureEditor/Private/TextureEditorToolkit.h
+++ b/Engine/Source/Editor/TextureEditor/Private/TextureEditorToolkit.h
@@ -530,7 +530,7 @@ private:
 	void OnEstimateCompressionChanged(ECheckBoxState NewState);
 	ECheckBoxState EstimateCompressionChecked() const;
 
-	TSharedPtr<STextBlock> OodleEncoderUsed;
+	TSharedPtr<STextBlock> OodleCompressorUsed;
 	TSharedPtr<STextBlock> OodleLevelUsed;
 	TSharedPtr<STextBlock> OodleCompressionBlockUsed;
 	TSharedPtr<STextBlock> OodleEstimateRaw;
