Fix texture border in the Texture Editor.

diff --git a/Engine/Source/Editor/TextureEditor/Private/Models/TextureEditorViewportClient.cpp b/Engine/Source/Editor/TextureEditor/Private/Models/TextureEditorViewportClient.cpp
--- a/Engine/Source/Editor/TextureEditor/Private/Models/TextureEditorViewportClient.cpp
+++ b/Engine/Source/Editor/TextureEditor/Private/Models/TextureEditorViewportClient.cpp
@@ -217,7 +217,8 @@ void FTextureEditorViewportClient::Draw(FViewport* Viewport, FCanvas* Canvas)
 		// Draw a white border around the texture to show its extents
 		if (Settings.TextureBorderEnabled)
 		{
-			FCanvasBoxItem BoxItem(FVector2D(XPos, YPos), FVector2D(Width + BorderSize, Height + BorderSize));
+			FCanvasBoxItem BoxItem(FVector2D(XPos - (BorderSize - 1) * 0.5f, YPos - (BorderSize - 1) * 0.5f), FVector2D(Width + BorderSize, Height + BorderSize));
+			BoxItem.LineThickness = BorderSize;
 			BoxItem.SetColor( Settings.TextureBorderColor );
 			Canvas->DrawItem( BoxItem );
 		}
