Texture Editor mouse wheel zoom should be a multiplicative factor not an additive factor
fixes zoom getting very slow as it gets larger

diff --git a/Engine/Source/Editor/TextureEditor/Private/TextureEditorConstants.h b/Engine/Source/Editor/TextureEditor/Private/TextureEditorConstants.h
--- a/Engine/Source/Editor/TextureEditor/Private/TextureEditorConstants.h
+++ b/Engine/Source/Editor/TextureEditor/Private/TextureEditorConstants.h
@@ -4,9 +4,12 @@
 
 #include "CoreMinimal.h"
 
+// MaxZoom and MinZoom should both be powers of two
 const double MaxZoom = 16.0;
-const double MinZoom = 0.01;
-const double ZoomStep = 0.1;
+const double MinZoom = 1.0/64;
+// ZoomFactor is multiplicative such that an integer number of steps will give a power of two zoom (50% or 200%)
+const int ZoomFactorLogSteps = 8;
+const double ZoomFactor = pow(2.0,1.0/ZoomFactorLogSteps);
 
 // Specifies the maximum allowed exposure bias.
 const int32 MaxExposure = 10;
diff --git a/Engine/Source/Editor/TextureEditor/Private/TextureEditorToolkit.cpp b/Engine/Source/Editor/TextureEditor/Private/TextureEditorToolkit.cpp
--- a/Engine/Source/Editor/TextureEditor/Private/TextureEditorToolkit.cpp
+++ b/Engine/Source/Editor/TextureEditor/Private/TextureEditorToolkit.cpp
@@ -889,42 +889,38 @@ double FTextureEditorToolkit::CalculateDisplayedZoomLevel() const
 
 void FTextureEditorToolkit::SetCustomZoomLevel( double ZoomValue )
 {
-	Zoom = FMath::Clamp(ZoomValue, MinZoom, MaxZoom);
+	// snap to discrete steps so that if we are nearly at 1.0 or 2.0, we hit them exactly:
+	//ZoomValue = FMath::GridSnap(ZoomValue, MinZoom/4.0);
+
+	double LogZoom = log2(ZoomValue);
+	// the mouse wheel zoom is quantized on ZoomFactorLogSteps
+	//	but that's too chunky for the drag slider, give it more steps, but on the same quantization grid
+	double QuantizationSteps = ZoomFactorLogSteps*2.0;
+	double LogZoomQuantized = (1.0/QuantizationSteps) * FMath::RoundToInt( QuantizationSteps * LogZoom );
+	ZoomValue = pow(2.0,LogZoomQuantized);
+
+	ZoomValue = FMath::Clamp(ZoomValue, MinZoom, MaxZoom);
 	
+	// set member variable "Zoom"
+	Zoom = ZoomValue;
+
 	// For now we also want to be in custom mode whenever this is changed
 	SetZoomMode(ETextureEditorZoomMode::Custom);
 }
 
 
-void FTextureEditorToolkit::OffsetZoom(double OffsetValue, bool bSnapToStepSize)
-{
-	// Offset from our current "visual" zoom level so that you can
-	// smoothly transition from Fit/Fill mode into a custom zoom level
-	const double CurrentZoom = CalculateDisplayedZoomLevel();
-
-	if (bSnapToStepSize)
-	{
-		// Snap to the zoom step when offsetting to avoid zooming all the way to the min (0.01)
-		// then back up (+0.1) causing your zoom level to be off by 0.01 (eg. 11%)
-		// If we were in a fit view mode then our current zoom level could also be off the grid
-		const double FinalZoom = FMath::GridSnap(CurrentZoom + OffsetValue, ZoomStep);
-		SetCustomZoomLevel(FinalZoom);
-	}
-	else
-	{
-		SetCustomZoomLevel(CurrentZoom + OffsetValue);
-	}
-}
-
 void FTextureEditorToolkit::ZoomIn( )
 {
-	OffsetZoom(ZoomStep);
+	// mouse wheel zoom
+	const double CurrentZoom = CalculateDisplayedZoomLevel();
+	SetCustomZoomLevel(CurrentZoom * ZoomFactor);
 }
 
 
 void FTextureEditorToolkit::ZoomOut( )
 {
-	OffsetZoom(-ZoomStep);
+	const double CurrentZoom = CalculateDisplayedZoomLevel();
+	SetCustomZoomLevel(CurrentZoom / ZoomFactor);
 }
 
 float FTextureEditorToolkit::GetVolumeOpacity() const
@@ -2592,14 +2588,22 @@ FText FTextureEditorToolkit::HandleZoomPercentageText() const
 	return ZoomLevelPercent;
 }
 
-void FTextureEditorToolkit::HandleZoomSliderChanged(float NewValue)
+void FTextureEditorToolkit::HandleZoomSliderChanged(float SliderValue)
 {
-	SetCustomZoomLevel(NewValue * MaxZoom);
+	// zoom slider is log scale, SliderValue in [0,1] between MinZoom and MaxZoom
+	double Octaves = log2( MaxZoom/MinZoom );
+	double ZoomValue = pow(2.0,SliderValue * Octaves) * MinZoom;
+
+	SetCustomZoomLevel((float)ZoomValue);
 }
 
 float FTextureEditorToolkit::HandleZoomSliderValue() const
 {
-	return (CalculateDisplayedZoomLevel() / MaxZoom);
+	float ZoomValue = CalculateDisplayedZoomLevel();
+	double Octaves = log2( MaxZoom/MinZoom );
+	double SliderValue = log2( ZoomValue/MinZoom ) / Octaves;
+
+	return (float)SliderValue;
 }
 
 int32 FTextureEditorToolkit::GetEditorOodleSettingsEffort() const
diff --git a/Engine/Source/Editor/TextureEditor/Private/TextureEditorToolkit.h b/Engine/Source/Editor/TextureEditor/Private/TextureEditorToolkit.h
--- a/Engine/Source/Editor/TextureEditor/Private/TextureEditorToolkit.h
+++ b/Engine/Source/Editor/TextureEditor/Private/TextureEditorToolkit.h
@@ -80,7 +80,6 @@ public:
 	virtual void SetZoomMode( const ETextureEditorZoomMode ZoomMode ) override;
 	virtual ETextureEditorZoomMode GetZoomMode() const override;
 	virtual double CalculateDisplayedZoomLevel() const override;
-	virtual void OffsetZoom( double OffsetValue, bool bSnapToStepSize = true );
 	virtual void ZoomIn( ) override;
 	virtual void ZoomOut( ) override;
 	virtual float GetVolumeOpacity( ) const override;
