Implement support for nearest-point texture sampling in the Texture Editor.

diff --git a/Engine/Plugins/VirtualProduction/CameraCalibration/Source/CameraCalibrationEditor/Private/AssetEditor/Viewport/SimulcamEditorViewportClient.cpp b/Engine/Plugins/VirtualProduction/CameraCalibration/Source/CameraCalibrationEditor/Private/AssetEditor/Viewport/SimulcamEditorViewportClient.cpp
--- a/Engine/Plugins/VirtualProduction/CameraCalibration/Source/CameraCalibrationEditor/Private/AssetEditor/Viewport/SimulcamEditorViewportClient.cpp
+++ b/Engine/Plugins/VirtualProduction/CameraCalibration/Source/CameraCalibrationEditor/Private/AssetEditor/Viewport/SimulcamEditorViewportClient.cpp
@@ -55,7 +55,7 @@ void FSimulcamEditorViewportClient::Draw(FViewport* Viewport, FCanvas* Canvas)
 	constexpr float SliceIndex = 0;
 
 	TRefCountPtr<FBatchedElementParameters> BatchedElementParameters;
-	BatchedElementParameters = new FBatchedElementTexture2DPreviewParameters(MipLevel, LayerIndex, SliceIndex, false, false, false, false, false);
+	BatchedElementParameters = new FBatchedElementTexture2DPreviewParameters(MipLevel, LayerIndex, SliceIndex, false, false, false, false, false, false);
 
 	if (Texture->GetResource() != nullptr)
 	{
diff --git a/Engine/Source/Editor/TextureEditor/Classes/TextureEditorSettings.h b/Engine/Source/Editor/TextureEditor/Classes/TextureEditorSettings.h
--- a/Engine/Source/Editor/TextureEditor/Classes/TextureEditorSettings.h
+++ b/Engine/Source/Editor/TextureEditor/Classes/TextureEditorSettings.h
@@ -20,6 +20,13 @@ enum ETextureEditorBackgrounds
 	TextureEditorBackground_CheckeredFill UMETA(DisplayName="Checkered (Fill)")
 };
 
+UENUM()
+enum ETextureEditorSampling
+{
+	TextureEditorSampling_Default UMETA(DisplayName = "Default Sampling"),
+	TextureEditorSampling_Point UMETA(DisplayName = "Nearest-Point Sampling"),
+};
+
 UENUM()
 enum ETextureEditorVolumeViewMode
 {
@@ -57,6 +64,10 @@ public:
 	UPROPERTY(config)
 	TEnumAsByte<ETextureEditorBackgrounds> Background;
 
+	/** The texture sampling mode used to render textures in the texture editor view port. */
+	UPROPERTY(config)
+	TEnumAsByte<ETextureEditorSampling> Sampling;
+
 	/** The view mode when previewing volume textures. */
 	UPROPERTY(config)
 	TEnumAsByte<ETextureEditorVolumeViewMode> VolumeViewMode;
diff --git a/Engine/Source/Editor/TextureEditor/Private/Menus/TextureEditorViewOptionsMenu.h b/Engine/Source/Editor/TextureEditor/Private/Menus/TextureEditorViewOptionsMenu.h
--- a/Engine/Source/Editor/TextureEditor/Private/Menus/TextureEditorViewOptionsMenu.h
+++ b/Engine/Source/Editor/TextureEditor/Private/Menus/TextureEditorViewOptionsMenu.h
@@ -34,6 +34,12 @@ public:
 			);
 
 			MenuBuilder.AddMenuEntry(FTextureEditorCommands::Get().TextureBorder);
+
+			MenuBuilder.AddSubMenu(
+				LOCTEXT("Sampling", "Sampling Mode"),
+				LOCTEXT("SamplingTooltip", "Set the texture sampling mode"),
+				FNewMenuDelegate::CreateStatic(&FTextureEditorViewOptionsMenu::GenerateSamplingMenuContent)
+			);
 		}
 		MenuBuilder.EndSection();
 
@@ -55,6 +61,16 @@ protected:
 		MenuBuilder.AddMenuEntry(FTextureEditorCommands::Get().SolidBackground);
 	}
 
+	/**
+	 * Creates the 'Sampling Mode' sub-menu.
+	 *
+	 * @param MenuBuilder The builder for the menu that owns this menu.
+	 */
+	static void GenerateSamplingMenuContent(FMenuBuilder& MenuBuilder)
+	{
+		MenuBuilder.AddMenuEntry(FTextureEditorCommands::Get().DefaultSampling);
+		MenuBuilder.AddMenuEntry(FTextureEditorCommands::Get().PointSampling);
+	}
 };
 
 
diff --git a/Engine/Source/Editor/TextureEditor/Private/Models/TextureEditorCommands.cpp b/Engine/Source/Editor/TextureEditor/Private/Models/TextureEditorCommands.cpp
--- a/Engine/Source/Editor/TextureEditor/Private/Models/TextureEditorCommands.cpp
+++ b/Engine/Source/Editor/TextureEditor/Private/Models/TextureEditorCommands.cpp
@@ -22,6 +22,8 @@ void FTextureEditorCommands::RegisterCommands()
 	UI_COMMAND(ZoomToNatural, "Scale To 100%", "Resets the zoom level to 100%, displaying the texture at its natural size", EUserInterfaceActionType::ToggleButton, FInputChord());
 	UI_COMMAND(SolidBackground, "Solid Color", "Solid color background", EUserInterfaceActionType::RadioButton, FInputChord());
 	UI_COMMAND(TextureBorder, "Draw Border", "If enabled, a border is drawn around the texture", EUserInterfaceActionType::ToggleButton, FInputChord());
+	UI_COMMAND(DefaultSampling, "Default Sampling", "Use default texture sampling", EUserInterfaceActionType::RadioButton, FInputChord());
+	UI_COMMAND(PointSampling, "Nearest-Point Sampling", "Use nearest-point texture sampling", EUserInterfaceActionType::RadioButton, FInputChord());
 
 	UI_COMMAND(CompressNow, "Compress", "Compress the texture", EUserInterfaceActionType::Button, FInputChord());
 	UI_COMMAND(Reimport, "Reimport", "Reimports the texture from file", EUserInterfaceActionType::Button, FInputChord());
diff --git a/Engine/Source/Editor/TextureEditor/Private/Models/TextureEditorCommands.h b/Engine/Source/Editor/TextureEditor/Private/Models/TextureEditorCommands.h
--- a/Engine/Source/Editor/TextureEditor/Private/Models/TextureEditorCommands.h
+++ b/Engine/Source/Editor/TextureEditor/Private/Models/TextureEditorCommands.h
@@ -70,6 +70,12 @@ public:
 	/** If enabled, a border is drawn around the texture */
 	TSharedPtr<FUICommandInfo> TextureBorder;
 
+	/** If enabled, default sampling is used for rendering the texture */
+	TSharedPtr<FUICommandInfo> DefaultSampling;
+
+	/** If enabled, nearest-point sampling is used for rendering the texture */
+	TSharedPtr<FUICommandInfo> PointSampling;
+
 	/** Compress the texture */
 	TSharedPtr<FUICommandInfo> CompressNow;
 
diff --git a/Engine/Source/Editor/TextureEditor/Private/Models/TextureEditorViewportClient.cpp b/Engine/Source/Editor/TextureEditor/Private/Models/TextureEditorViewportClient.cpp
--- a/Engine/Source/Editor/TextureEditor/Private/Models/TextureEditorViewportClient.cpp
+++ b/Engine/Source/Editor/TextureEditor/Private/Models/TextureEditorViewportClient.cpp
@@ -104,6 +104,7 @@ void FTextureEditorViewportClient::Draw(FViewport* Viewport, FCanvas* Canvas)
 	const float MipLevel = (float)TextureEditorPinned->GetMipLevel();
 	const float LayerIndex = (float)TextureEditorPinned->GetLayer();
 	const float SliceIndex = (float)TextureEditorPinned->GetSlice();
+	const bool bUsePointSampling = TextureEditorPinned->GetSampling() == ETextureEditorSampling::TextureEditorSampling_Point;
 
 	bool bIsVirtualTexture = false;
 
@@ -128,7 +129,7 @@ void FTextureEditorViewportClient::Draw(FViewport* Viewport, FCanvas* Canvas)
 				FaceIndex == 4 ? FMatrix44f(FVector3f::UpVector, FVector3f::ForwardVector, FVector3f::RightVector, IdentityW) :
 				FMatrix44f(FVector3f::DownVector, FVector3f::BackwardVector, FVector3f::RightVector, IdentityW);
 			const bool bShowLongLatUnwrap = TextureEditorPinned->GetCubemapViewMode() == TextureEditorCubemapViewMode_2DView && FaceIndex < 0;
-			BatchedElementParameters = new FMipLevelBatchedElementParameters(MipLevel, SliceIndex, TextureCubeArray != nullptr, ViewMatrix, bShowLongLatUnwrap, false);
+			BatchedElementParameters = new FMipLevelBatchedElementParameters(MipLevel, SliceIndex, TextureCubeArray != nullptr, ViewMatrix, bShowLongLatUnwrap, false, bUsePointSampling);
 		}
 		else if (VolumeTexture)
 		{
@@ -138,7 +139,8 @@ void FTextureEditorViewportClient::Draw(FViewport* Viewport, FCanvas* Canvas)
 				MipLevel, 
 				(float)TextureEditorPinned->GetVolumeOpacity(),
 				true, 
-				TextureEditorPinned->GetOrientation());
+				TextureEditorPinned->GetOrientation(),
+				bUsePointSampling);
 		}
 		else if (RTTextureVolume)
 		{
@@ -148,7 +150,8 @@ void FTextureEditorViewportClient::Draw(FViewport* Viewport, FCanvas* Canvas)
 				MipLevel,
 				(float)TextureEditorPinned->GetVolumeOpacity(),
 				true,
-				TextureEditorPinned->GetOrientation());
+				TextureEditorPinned->GetOrientation(),
+				bUsePointSampling);
 		}
 		else if (Texture2D)
 		{
@@ -156,26 +159,26 @@ void FTextureEditorViewportClient::Draw(FViewport* Viewport, FCanvas* Canvas)
 			bool bIsSingleChannel = Texture2D->CompressionSettings == TC_Grayscale || Texture2D->CompressionSettings == TC_Alpha;
 			bool bSingleVTPhysicalSpace = Texture2D->IsVirtualTexturedWithSinglePhysicalSpace();
 			bIsVirtualTexture = Texture2D->IsCurrentlyVirtualTextured();
-			BatchedElementParameters = new FBatchedElementTexture2DPreviewParameters(MipLevel, LayerIndex, SliceIndex, bIsNormalMap, bIsSingleChannel, bSingleVTPhysicalSpace, bIsVirtualTexture, false);
+			BatchedElementParameters = new FBatchedElementTexture2DPreviewParameters(MipLevel, LayerIndex, SliceIndex, bIsNormalMap, bIsSingleChannel, bSingleVTPhysicalSpace, bIsVirtualTexture, false, bUsePointSampling);
 		}
 		else if (Texture2DArray) 
 		{
 			bool bIsNormalMap = Texture2DArray->IsNormalMap();
 			bool bIsSingleChannel = Texture2DArray->CompressionSettings == TC_Grayscale || Texture2DArray->CompressionSettings == TC_Alpha;
-			BatchedElementParameters = new FBatchedElementTexture2DPreviewParameters(MipLevel, LayerIndex, SliceIndex, bIsNormalMap, bIsSingleChannel, false, false, true);
+			BatchedElementParameters = new FBatchedElementTexture2DPreviewParameters(MipLevel, LayerIndex, SliceIndex, bIsNormalMap, bIsSingleChannel, false, false, true, bUsePointSampling);
 		}
 		else if (TextureRT2D)
 		{
-			BatchedElementParameters = new FBatchedElementTexture2DPreviewParameters(MipLevel, LayerIndex, SliceIndex, false, false, false, false, false);
+			BatchedElementParameters = new FBatchedElementTexture2DPreviewParameters(MipLevel, LayerIndex, SliceIndex, false, false, false, false, false, bUsePointSampling);
 		}
 		else if (TextureRT2DArray)
 		{
-			BatchedElementParameters = new FBatchedElementTexture2DPreviewParameters(MipLevel, LayerIndex, SliceIndex, false, false, false, false, true);
+			BatchedElementParameters = new FBatchedElementTexture2DPreviewParameters(MipLevel, LayerIndex, SliceIndex, false, false, false, false, true, bUsePointSampling);
 		}
 		else
 		{
 			// Default to treating any UTexture derivative as a 2D texture resource
-			BatchedElementParameters = new FBatchedElementTexture2DPreviewParameters(MipLevel, LayerIndex, SliceIndex, false, false, false, false, false);
+			BatchedElementParameters = new FBatchedElementTexture2DPreviewParameters(MipLevel, LayerIndex, SliceIndex, false, false, false, false, false, bUsePointSampling);
 		}
 	}
 
diff --git a/Engine/Source/Editor/TextureEditor/Private/TextureEditorClasses.cpp b/Engine/Source/Editor/TextureEditor/Private/TextureEditorClasses.cpp
--- a/Engine/Source/Editor/TextureEditor/Private/TextureEditorClasses.cpp
+++ b/Engine/Source/Editor/TextureEditor/Private/TextureEditorClasses.cpp
@@ -10,6 +10,7 @@ class FObjectInitializer;
 UTextureEditorSettings::UTextureEditorSettings( const FObjectInitializer& ObjectInitializer )
 	: Super(ObjectInitializer)
 	, Background(TextureEditorBackground_Checkered)
+	, Sampling(TextureEditorSampling_Default)
 	, BackgroundColor(FColor::FromHex("242424FF"))
 	, CheckerColorOne(FColor(128, 128, 128))
 	, CheckerColorTwo(FColor(64, 64, 64))
diff --git a/Engine/Source/Editor/TextureEditor/Private/TextureEditorToolkit.cpp b/Engine/Source/Editor/TextureEditor/Private/TextureEditorToolkit.cpp
--- a/Engine/Source/Editor/TextureEditor/Private/TextureEditorToolkit.cpp
+++ b/Engine/Source/Editor/TextureEditor/Private/TextureEditorToolkit.cpp
@@ -262,11 +262,12 @@ void FTextureEditorToolkit::InitTextureEditor( const EToolkitMode::Type Mode, co
 	SpecifiedFace = 0;
 	bUseSpecifiedFace = false;
 
-	// Start at whatever the last used zoom mode, volume view mode and cubemap view mode were
+	// Start at whatever the last used values of the following settings were
 	const UTextureEditorSettings& Settings = *GetDefault<UTextureEditorSettings>();
 	ZoomMode = Settings.ZoomMode;
 	VolumeViewMode = Settings.VolumeViewMode;
 	CubemapViewMode = Settings.CubemapViewMode;
+	Sampling = Settings.Sampling;
 
 	ResetOrientation();
 	Zoom = 1.0f;
@@ -988,6 +989,11 @@ void FTextureEditorToolkit::ResetOrientation()
 	SetOrientation(IsVolumeTexture() ? FRotator(90, 0, -90) : FRotator(0, 0, 0));
 }
 
+ETextureEditorSampling FTextureEditorToolkit::GetSampling() const
+{
+	return Sampling;
+}
+
 /* IToolkit interface
  *****************************************************************************/
 
@@ -1109,6 +1115,18 @@ void FTextureEditorToolkit::BindCommands( )
 		FCanExecuteAction(),
 		FIsActionChecked::CreateSP(this, &FTextureEditorToolkit::HandleTextureBorderActionIsChecked));
 
+	ToolkitCommands->MapAction(
+		Commands.DefaultSampling,
+		FExecuteAction::CreateSP(this, &FTextureEditorToolkit::HandleSamplingActionExecute, TextureEditorSampling_Default),
+		FCanExecuteAction(),
+		FIsActionChecked::CreateSP(this, &FTextureEditorToolkit::HandleSamplingActionIsChecked, TextureEditorSampling_Default));
+
+	ToolkitCommands->MapAction(
+		Commands.PointSampling,
+		FExecuteAction::CreateSP(this, &FTextureEditorToolkit::HandleSamplingActionExecute, TextureEditorSampling_Point),
+		FCanExecuteAction(),
+		FIsActionChecked::CreateSP(this, &FTextureEditorToolkit::HandleSamplingActionIsChecked, TextureEditorSampling_Point));
+
 	ToolkitCommands->MapAction(
 		Commands.CompressNow,
 		FExecuteAction::CreateSP(this, &FTextureEditorToolkit::HandleCompressNowActionExecute),
@@ -2139,6 +2157,22 @@ bool FTextureEditorToolkit::HandleCheckeredBackgroundActionIsChecked( ETextureEd
 	return (Background == Settings.Background);
 }
 
+void FTextureEditorToolkit::HandleSamplingActionExecute(ETextureEditorSampling InSampling)
+{
+	// Update our own sampling
+	Sampling = InSampling;
+
+	UTextureEditorSettings& Settings = *GetMutableDefault<UTextureEditorSettings>();
+	Settings.Sampling = InSampling;
+	Settings.PostEditChange();
+}
+
+
+bool FTextureEditorToolkit::HandleSamplingActionIsChecked(ETextureEditorSampling InSampling)
+{
+	return InSampling == Sampling;
+}
+
 // Callback for toggling the volume view action.
 void FTextureEditorToolkit::HandleVolumeViewModeActionExecute(ETextureEditorVolumeViewMode InViewMode)
 {
diff --git a/Engine/Source/Editor/TextureEditor/Private/TextureEditorToolkit.h b/Engine/Source/Editor/TextureEditor/Private/TextureEditorToolkit.h
--- a/Engine/Source/Editor/TextureEditor/Private/TextureEditorToolkit.h
+++ b/Engine/Source/Editor/TextureEditor/Private/TextureEditorToolkit.h
@@ -92,6 +92,7 @@ public:
 	virtual const FRotator& GetOrientation() const override;
 	virtual void SetOrientation(const FRotator& InOrientation) override;
 	virtual void ResetOrientation() override;
+	virtual ETextureEditorSampling GetSampling() const override;
 	virtual int32 GetExposureBias() const override
 	{
 		return ExposureBias;
@@ -201,6 +202,12 @@ private:
 	// Callback for getting the checked state of the Checkered Background action.
 	bool HandleCheckeredBackgroundActionIsChecked( ETextureEditorBackgrounds Background );
 
+	// Callback for toggling the Sampling action.
+	void HandleSamplingActionExecute(ETextureEditorSampling Sampling);
+
+	// Callback for getting the checked state of the Sampling action.
+	bool HandleSamplingActionIsChecked(ETextureEditorSampling Sampling);
+
 	// Callback for toggling the volume view action.
 	void HandleVolumeViewModeActionExecute( ETextureEditorVolumeViewMode InViewMode );
 
@@ -487,6 +494,9 @@ private:
 	/** This toolkit's current cubemap view mode **/
 	ETextureEditorCubemapViewMode CubemapViewMode;
 
+	/** This toolkit's current sampling mode **/
+	ETextureEditorSampling Sampling;
+
 	// Orientation of the texture in 3d when using "3D View" mode
 	FRotator Orientation;
 
diff --git a/Engine/Source/Editor/TextureEditor/Public/Interfaces/ITextureEditorToolkit.h b/Engine/Source/Editor/TextureEditor/Public/Interfaces/ITextureEditorToolkit.h
--- a/Engine/Source/Editor/TextureEditor/Public/Interfaces/ITextureEditorToolkit.h
+++ b/Engine/Source/Editor/TextureEditor/Public/Interfaces/ITextureEditorToolkit.h
@@ -55,6 +55,7 @@ public:
 	virtual const FRotator& GetOrientation() const = 0;
 	virtual void SetOrientation(const FRotator& InOrientation) = 0;
 	virtual void ResetOrientation() = 0;
+	virtual ETextureEditorSampling GetSampling() const = 0;
 	virtual int32 GetExposureBias() const = 0;
 	virtual bool IsVolumeTexture() const = 0;
 public:
diff --git a/Engine/Source/Editor/UnrealEd/Private/Texture2DPreview.cpp b/Engine/Source/Editor/UnrealEd/Private/Texture2DPreview.cpp
--- a/Engine/Source/Editor/UnrealEd/Private/Texture2DPreview.cpp
+++ b/Engine/Source/Editor/UnrealEd/Private/Texture2DPreview.cpp
@@ -63,7 +63,7 @@ public:
 		return IsFeatureLevelSupported(Parameters.Platform, ERHIFeatureLevel::SM5) && !IsConsolePlatform(Parameters.Platform);
 	}
 
-	void SetParameters(FRHICommandList& InRHICmdList, const FTexture* InTextureValue, const FMatrix44f& InColorWeightsValue, float InGammaValue, float InMipLevel, float InLayerIndex, float InSliceIndex, bool bInIsNormalMap, bool bInIsSingleVTPhysicalSpace, bool bInIsVirtualTexture, bool bInIsTextureArray)
+	void SetParameters(FRHICommandList& InRHICmdList, const FTexture* InTextureValue, const FMatrix44f& InColorWeightsValue, float InGammaValue, float InMipLevel, float InLayerIndex, float InSliceIndex, bool bInIsNormalMap, bool bInIsSingleVTPhysicalSpace, bool bInIsVirtualTexture, bool bInIsTextureArray, bool bInUsePointSampling)
 	{
 		FRHIPixelShader* ShaderRHI = InRHICmdList.GetBoundPixelShader();
 		if (bInIsVirtualTexture)
@@ -73,7 +73,9 @@ public:
 	
 			FRHIShaderResourceView* PhysicalView = AllocatedVT->GetPhysicalTextureSRV((uint32)InLayerIndex, InTextureValue->bSRGB);
 			SetSRVParameter(InRHICmdList, ShaderRHI, InTexture, PhysicalView);
-			SetSamplerParameter(InRHICmdList, ShaderRHI, InTextureSampler, VirtualTextureValue->SamplerStateRHI);
+
+			FRHISamplerState* SamplerState = bInUsePointSampling ? TStaticSamplerState<SF_Point>::GetRHI() : VirtualTextureValue->SamplerStateRHI.GetReference();
+			SetSamplerParameter(InRHICmdList, ShaderRHI, InTextureSampler, SamplerState);
 
 			SetTextureParameter(InRHICmdList, ShaderRHI, InPageTableTexture0, AllocatedVT->GetPageTableTexture(0u));
 			if (AllocatedVT->GetNumPageTableTextures() > 1u)
@@ -98,7 +100,9 @@ public:
 		{
 			SetTextureParameter(InRHICmdList, ShaderRHI, InPageTableTexture0, GBlackTexture->TextureRHI);
 			SetTextureParameter(InRHICmdList, ShaderRHI, InPageTableTexture1, GBlackTexture->TextureRHI);
-			SetTextureParameter(InRHICmdList, ShaderRHI, InTexture, InTextureSampler, InTextureValue);
+
+			FRHISamplerState* SamplerState = bInUsePointSampling ? TStaticSamplerState<SF_Point>::GetRHI() : InTextureValue->SamplerStateRHI.GetReference();
+			SetTextureParameter(InRHICmdList, ShaderRHI, InTexture, InTextureSampler, SamplerState, InTextureValue->TextureRHI);
 		}
 		
 		SetShaderValue(InRHICmdList, ShaderRHI,ColorWeights,InColorWeightsValue);
@@ -165,5 +169,5 @@ void FBatchedElementTexture2DPreviewParameters::BindShaders(
 	SetGraphicsPipelineState(RHICmdList, GraphicsPSOInit, 0);
 
 	VertexShader->SetParameters(RHICmdList, InTransform);
-	PixelShader->SetParameters(RHICmdList, Texture, FMatrix44f(ColorWeights), InGamma, MipLevel, LayerIndex, SliceIndex, bIsNormalMap, bIsSingleVTPhysicalSpace, bIsVirtualTexture, bIsTextureArray);
+	PixelShader->SetParameters(RHICmdList, Texture, FMatrix44f(ColorWeights), InGamma, MipLevel, LayerIndex, SliceIndex, bIsNormalMap, bIsSingleVTPhysicalSpace, bIsVirtualTexture, bIsTextureArray, bUsePointSampling);
 }
diff --git a/Engine/Source/Editor/UnrealEd/Private/ThumbnailRendering/TextureThumbnailRenderer.cpp b/Engine/Source/Editor/UnrealEd/Private/ThumbnailRendering/TextureThumbnailRenderer.cpp
--- a/Engine/Source/Editor/UnrealEd/Private/ThumbnailRendering/TextureThumbnailRenderer.cpp
+++ b/Engine/Source/Editor/UnrealEd/Private/ThumbnailRendering/TextureThumbnailRenderer.cpp
@@ -75,7 +75,7 @@ void UTextureThumbnailRenderer::Draw(UObject* Object, int32 X, int32 Y, uint32 W
 		if(TextureCube || TextureCubeArray || RTTextureCube)
 		{
 			// is released by the render thread when it was rendered
-			BatchedElementParameters = new FMipLevelBatchedElementParameters((float)0, (float)-1, TextureCubeArray != nullptr, FMatrix44f::Identity, true, false);
+			BatchedElementParameters = new FMipLevelBatchedElementParameters((float)0, (float)-1, TextureCubeArray != nullptr, FMatrix44f::Identity, true, false, false);
 			
 			// If the thumbnail is square then make it 2:1 for cubes.
 			if(Width == Height)
@@ -88,7 +88,7 @@ void UTextureThumbnailRenderer::Draw(UObject* Object, int32 X, int32 Y, uint32 W
 		{
 			bool bIsNormalMap = Texture2DArray->IsNormalMap();
 			bool bIsSingleChannel = true;
-			BatchedElementParameters = new FBatchedElementTexture2DPreviewParameters((float)0, (float)0, (float)-1, bIsNormalMap, bIsSingleChannel, false, false, true);
+			BatchedElementParameters = new FBatchedElementTexture2DPreviewParameters((float)0, (float)0, (float)-1, bIsNormalMap, bIsSingleChannel, false, false, true, false);
 		}
 		else if (TextureLightProfile)
 		{
diff --git a/Engine/Source/Editor/UnrealEd/Private/VolumeTexturePreview.cpp b/Engine/Source/Editor/UnrealEd/Private/VolumeTexturePreview.cpp
--- a/Engine/Source/Editor/UnrealEd/Private/VolumeTexturePreview.cpp
+++ b/Engine/Source/Editor/UnrealEd/Private/VolumeTexturePreview.cpp
@@ -76,11 +76,13 @@ public:
 		return IsFeatureLevelSupported(Parameters.Platform, ERHIFeatureLevel::SM5) && !IsConsolePlatform(Parameters.Platform);
 	}
 	
-	void SetParameters(FRHICommandList& RHICmdList, const FTexture* TextureValue, int32 SizeZ, const FMatrix44f& ColorWeightsValue, float GammaValue, float MipLevel, float Opacity, const FRotator& TraceOrientation)
+	void SetParameters(FRHICommandList& RHICmdList, const FTexture* TextureValue, int32 SizeZ, const FMatrix44f& ColorWeightsValue, float GammaValue, float MipLevel, float Opacity, const FRotator& TraceOrientation, bool bUsePointSampling)
 	{
 		FRHIPixelShader* ShaderRHI = RHICmdList.GetBoundPixelShader();
 
-		SetTextureParameter(RHICmdList, ShaderRHI, InTexture, InTextureSampler, TextureValue);
+		FRHISamplerState* SamplerState = bUsePointSampling ? TStaticSamplerState<SF_Point>::GetRHI() : TextureValue->SamplerStateRHI.GetReference();
+		SetTextureParameter(RHICmdList, ShaderRHI, InTexture, InTextureSampler, SamplerState, TextureValue->TextureRHI);
+
 		if (GEditor && GEditor->Bad)
 		{
 			SetTextureParameter(RHICmdList, ShaderRHI, BadTexture, BadTextureSampler, GEditor->Bad->GetResource());
@@ -204,5 +206,5 @@ void FBatchedElementVolumeTexturePreviewParameters::BindShaders(
 	SetGraphicsPipelineState(RHICmdList, GraphicsPSOInit, 0);
 
 	VertexShader->SetParameters(RHICmdList, InTransform);
-	PixelShader->SetParameters(RHICmdList, Texture, SizeZ, ColorWeights, InGamma, MipLevel, Opacity, TraceOrientation);
+	PixelShader->SetParameters(RHICmdList, Texture, SizeZ, ColorWeights, InGamma, MipLevel, Opacity, TraceOrientation, bUsePointSampling);
 }
diff --git a/Engine/Source/Editor/UnrealEd/Public/Texture2DPreview.h b/Engine/Source/Editor/UnrealEd/Public/Texture2DPreview.h
--- a/Engine/Source/Editor/UnrealEd/Public/Texture2DPreview.h
+++ b/Engine/Source/Editor/UnrealEd/Public/Texture2DPreview.h
@@ -15,7 +15,7 @@
 class UNREALED_API FBatchedElementTexture2DPreviewParameters : public FBatchedElementParameters
 {
 public:
-	FBatchedElementTexture2DPreviewParameters(float InMipLevel, float InLayerIndex, float InSliceIndex, bool bInIsNormalMap, bool bInIsSingleChannel, bool bInIsSingleVTPhysicalSpace, bool bInIsVirtualTexture, bool bInIsTextureArray)
+	FBatchedElementTexture2DPreviewParameters(float InMipLevel, float InLayerIndex, float InSliceIndex, bool bInIsNormalMap, bool bInIsSingleChannel, bool bInIsSingleVTPhysicalSpace, bool bInIsVirtualTexture, bool bInIsTextureArray, bool bInUsePointSampling)
 		: MipLevel(InMipLevel)
 		, LayerIndex(InLayerIndex)
 		, SliceIndex(InSliceIndex)
@@ -24,6 +24,7 @@ public:
 		, bIsSingleVTPhysicalSpace(bInIsSingleVTPhysicalSpace)
 		, bIsVirtualTexture(bInIsVirtualTexture)
 		, bIsTextureArray(bInIsTextureArray)
+		, bUsePointSampling(bInUsePointSampling)
 	{
 	}
 
@@ -43,4 +44,7 @@ private:
 	/** Parameters that are used to select a shader permutation */
 	bool bIsVirtualTexture;
 	bool bIsTextureArray;
+
+	/** Whether to use nearest-point sampling when rendering the texture */
+	bool bUsePointSampling;
 };
diff --git a/Engine/Source/Editor/UnrealEd/Public/VolumeTexturePreview.h b/Engine/Source/Editor/UnrealEd/Public/VolumeTexturePreview.h
--- a/Engine/Source/Editor/UnrealEd/Public/VolumeTexturePreview.h
+++ b/Engine/Source/Editor/UnrealEd/Public/VolumeTexturePreview.h
@@ -15,13 +15,14 @@
 class UNREALED_API FBatchedElementVolumeTexturePreviewParameters : public FBatchedElementParameters
 {
 public:
-	FBatchedElementVolumeTexturePreviewParameters(bool InViewModeAsDepthSlices, int32 InSizeZ, float InMipLevel, float InOpacity, bool InShowSlices, const FRotator& InTraceOrientation)
+	FBatchedElementVolumeTexturePreviewParameters(bool InViewModeAsDepthSlices, int32 InSizeZ, float InMipLevel, float InOpacity, bool InShowSlices, const FRotator& InTraceOrientation, bool bInUsePointSampling)
 		: bViewModeAsDepthSlices(InViewModeAsDepthSlices)
 		, SizeZ(InSizeZ)
 		, MipLevel(InMipLevel)
 		, Opacity(InOpacity)
 		, bShowSlices(InShowSlices)
 		, TraceOrientation(InTraceOrientation)
+		, bUsePointSampling(bInUsePointSampling)
 	{
 	}
 
@@ -46,4 +47,7 @@ private:
 
 	/** The orientation when tracing */
 	FRotator TraceOrientation;
+
+	/** Whether to use nearest-point sampling when rendering the volume */
+	bool bUsePointSampling;
 };
diff --git a/Engine/Source/Runtime/Engine/Private/CubemapUnwrapUtils.cpp b/Engine/Source/Runtime/Engine/Private/CubemapUnwrapUtils.cpp
--- a/Engine/Source/Runtime/Engine/Private/CubemapUnwrapUtils.cpp
+++ b/Engine/Source/Runtime/Engine/Private/CubemapUnwrapUtils.cpp
@@ -37,7 +37,7 @@ namespace CubemapHelpers
 	bool GenerateLongLatUnwrap(const FTextureResource* TextureResource, const uint32 AxisDimenion, const EPixelFormat SourcePixelFormat, TArray64<uint8>& BitsOUT, FIntPoint& SizeOUT, EPixelFormat& FormatOUT)
 	{
 		TRefCountPtr<FBatchedElementParameters> BatchedElementParameters;
-		BatchedElementParameters = new FMipLevelBatchedElementParameters((float)0, (float)-1, false, FMatrix44f::Identity, true, true);
+		BatchedElementParameters = new FMipLevelBatchedElementParameters((float)0, (float)-1, false, FMatrix44f::Identity, true, true, false);
 		const FIntPoint LongLatDimensions(AxisDimenion * 2, AxisDimenion);
 
 		// If the source format is 8 bit per channel or less then select a LDR target format.
@@ -148,11 +148,12 @@ bool FCubemapTexturePropertiesPS::ShouldCompilePermutation(const FGlobalShaderPe
 	return true;
 }
 
-void FCubemapTexturePropertiesPS::SetParameters(FRHICommandList& RHICmdList, const FTexture* InTexture, const FMatrix& InColorWeightsValue, float InMipLevel, float InSliceIndex, bool bInIsTextureCubeArray, const FMatrix44f& InViewMatrix, bool bInShowLongLatUnwrap, float InGammaValue)
+void FCubemapTexturePropertiesPS::SetParameters(FRHICommandList& RHICmdList, const FTexture* InTexture, const FMatrix& InColorWeightsValue, float InMipLevel, float InSliceIndex, bool bInIsTextureCubeArray, const FMatrix44f& InViewMatrix, bool bInShowLongLatUnwrap, float InGammaValue, bool bInUsePointSampling)
 {
 	FRHIPixelShader* ShaderRHI = RHICmdList.GetBoundPixelShader();
 
-	SetTextureParameter(RHICmdList, ShaderRHI, CubeTexture, CubeTextureSampler, InTexture);
+	FRHISamplerState* SamplerState = bInUsePointSampling ? TStaticSamplerState<SF_Point>::GetRHI() : InTexture->SamplerStateRHI.GetReference();
+	SetTextureParameter(RHICmdList, ShaderRHI, CubeTexture, CubeTextureSampler, SamplerState, InTexture->TextureRHI);
 
 	FVector4f PackedProperties0Value(InMipLevel, bInShowLongLatUnwrap ? 1.0f : -1.0f, 0, 0);
 	SetShaderValue(RHICmdList, ShaderRHI, PackedProperties0, PackedProperties0Value);
@@ -192,7 +193,7 @@ void FMipLevelBatchedElementParameters::BindShaders(FRHICommandList& RHICmdList,
 	SetGraphicsPipelineState(RHICmdList, GraphicsPSOInit, 0);
 
 	VertexShader->SetParameters(RHICmdList, InTransform);
-	PixelShader->SetParameters(RHICmdList, Texture, ColorWeights, MipLevel, SliceIndex, bIsTextureCubeArray, ViewMatrix, bShowLongLatUnwrap, InGamma);
+	PixelShader->SetParameters(RHICmdList, Texture, ColorWeights, MipLevel, SliceIndex, bIsTextureCubeArray, ViewMatrix, bShowLongLatUnwrap, InGamma, bUsePointSampling);
 }
 
 void FIESLightProfilePS::SetParameters( FRHICommandList& RHICmdList, const FTexture* Texture, float InBrightnessInLumens )
diff --git a/Engine/Source/Runtime/Engine/Public/CubemapUnwrapUtils.h b/Engine/Source/Runtime/Engine/Public/CubemapUnwrapUtils.h
--- a/Engine/Source/Runtime/Engine/Public/CubemapUnwrapUtils.h
+++ b/Engine/Source/Runtime/Engine/Public/CubemapUnwrapUtils.h
@@ -79,7 +79,7 @@ public:
 	FCubemapTexturePropertiesPS(const ShaderMetaType::CompiledShaderInitializerType& Initializer);
 
 	static bool ShouldCompilePermutation(const FGlobalShaderPermutationParameters& Parameters);
-	void SetParameters(FRHICommandList& RHICmdList, const FTexture* InTexture, const FMatrix& InColorWeightsValue, float InMipLevel, float InSliceIndex, bool bInIsTextureCubeArray, const FMatrix44f& InViewMatrix, bool bInShowLongLatUnwrap, float InGammaValue);
+	void SetParameters(FRHICommandList& RHICmdList, const FTexture* InTexture, const FMatrix& InColorWeightsValue, float InMipLevel, float InSliceIndex, bool bInIsTextureCubeArray, const FMatrix44f& InViewMatrix, bool bInShowLongLatUnwrap, float InGammaValue, bool bInUsePointSampling);
 
 private:
 	LAYOUT_FIELD(FShaderResourceParameter, CubeTexture);
@@ -96,13 +96,14 @@ private:
 class ENGINE_API FMipLevelBatchedElementParameters : public FBatchedElementParameters
 {
 public:
-	FMipLevelBatchedElementParameters(float InMipLevel, float InSliceIndex, bool bInIsTextureCubeArray, const FMatrix44f& InViewMatrix, bool bInShowLongLatUnwrap, bool bInHDROutput)
+	FMipLevelBatchedElementParameters(float InMipLevel, float InSliceIndex, bool bInIsTextureCubeArray, const FMatrix44f& InViewMatrix, bool bInShowLongLatUnwrap, bool bInHDROutput, bool bInUsePointSampling)
 		: bHDROutput(bInHDROutput)
 		, MipLevel(InMipLevel)
 		, SliceIndex(InSliceIndex)
 		, ViewMatrix(InViewMatrix)
 		, bShowLongLatUnwrap(bInShowLongLatUnwrap)
 		, bIsTextureCubeArray(bInIsTextureCubeArray)
+		, bUsePointSampling(bInUsePointSampling)
 	{
 	}
 
@@ -120,6 +121,9 @@ private:
 
 	/** Parameters that are used to select a shader permutation */
 	bool bIsTextureCubeArray;
+
+	/** Whether to use nearest-point sampling when rendering the cubemap */
+	bool bUsePointSampling;
 };
 
 
