Opening a texture in the texture viewer always shows Final encoding.

diff --git a/Engine/Source/Editor/TextureEditor/Private/TextureEditorToolkit.cpp b/Engine/Source/Editor/TextureEditor/Private/TextureEditorToolkit.cpp
--- a/Engine/Source/Editor/TextureEditor/Private/TextureEditorToolkit.cpp
+++ b/Engine/Source/Editor/TextureEditor/Private/TextureEditorToolkit.cpp
@@ -317,6 +317,14 @@ void FTextureEditorToolkit::InitTextureEditor( const EToolkitMode::Type Mode, co
 
 	RegenerateMenusAndToolbars();
 
+	// If we're opening the texture editor then we likely want to see the Final
+	// texture quality for evaluation and parameter adjustment. This ends up latching
+	// for the current editor session, but since it's already getting built this is
+	// fine. This doesn't need the normal Modify() / PreEditChange/PostEditChange incantations
+	// because it's transient, and any async build was completed above.
+	Texture->CompressFinal = true;
+	PostTextureRecode();
+
 	// @todo toolkit world centric editing
 	/*if(IsWorldCentricAssetEditor())
 	{
