TextureEditor: Remove unused code in STextureEditorViewport

diff --git a/Engine/Source/Editor/TextureEditor/Private/Widgets/STextureEditorViewport.cpp b/Engine/Source/Editor/TextureEditor/Private/Widgets/STextureEditorViewport.cpp
--- a/Engine/Source/Editor/TextureEditor/Private/Widgets/STextureEditorViewport.cpp
+++ b/Engine/Source/Editor/TextureEditor/Private/Widgets/STextureEditorViewport.cpp
@@ -2,26 +2,15 @@
 
 #include "Widgets/STextureEditorViewport.h"
 #include "Framework/Application/SlateApplication.h"
-#include "Widgets/Text/STextBlock.h"
-#include "Framework/MultiBox/MultiBoxBuilder.h"
 #include "Widgets/Layout/SScrollBar.h"
 #include "Widgets/SToolTip.h"
-#include "Widgets/Input/SComboButton.h"
 #include "Widgets/SViewport.h"
-#include "Widgets/Input/SSlider.h"
-#include "Engine/Texture.h"
-#include "Engine/VolumeTexture.h"
-#include "Engine/TextureRenderTargetVolume.h"
 #include "Slate/SceneViewport.h"
-#include "TextureEditorConstants.h"
 #include "Widgets/Input/SNumericEntryBox.h"
-#include "TextureEditorSettings.h"
-
 
 #define LOCTEXT_NAMESPACE "STextureEditorViewport"
 
 
-
 /* STextureEditorViewport interface
  *****************************************************************************/
 
@@ -30,23 +19,10 @@ void STextureEditorViewport::AddReferencedObjects( FReferenceCollector& Collecto
 	ViewportClient->AddReferencedObjects(Collector);
 }
 
-
 void STextureEditorViewport::Construct( const FArguments& InArgs, const TSharedRef<ITextureEditorToolkit>& InToolkit )
 {
-
 	bIsRenderingEnabled = true;
 	ToolkitPtr = InToolkit;
-	
-	FText TextureName = FText::GetEmpty();
-
-	if (InToolkit->GetTexture() != nullptr)
-	{
-		FText FormattedText = InToolkit->HasValidTextureResource() ? FText::FromString(TEXT("{0}")) : LOCTEXT( "InvalidTextureWithParam", "{0} (Invalid Texture)");
-		TextureName = FText::Format(FormattedText, FText::FromName(InToolkit->GetTexture()->GetFName()));
-	}
-
-
-	TSharedPtr<SHorizontalBox> HorizontalBox;
 
 	this->ChildSlot
 	[
@@ -86,8 +62,6 @@ void STextureEditorViewport::Construct( const FArguments& InArgs, const TSharedR
 		]
 
 	];
-
-
 	
 	ViewportClient = MakeShareable(new FTextureEditorViewportClient(ToolkitPtr, SharedThis(this)));
 
@@ -97,7 +71,6 @@ void STextureEditorViewport::Construct( const FArguments& InArgs, const TSharedR
 	ViewportWidget->SetViewportInterface( Viewport.ToSharedRef() );
 }
 
-
 void STextureEditorViewport::ModifyCheckerboardTextureColors( )
 {
 	if (ViewportClient.IsValid())
@@ -161,17 +134,15 @@ FText STextureEditorViewport::GetDisplayedResolution( ) const
 /* STextureEditorViewport event handlers
  *****************************************************************************/
 
-
 void STextureEditorViewport::HandleHorizontalScrollBarScrolled( float InScrollOffsetFraction )
 {
-	float Ratio = ViewportClient->GetViewportHorizontalScrollBarRatio();
-	float MaxOffset = (Ratio < 1.0f) ? 1.0f - Ratio : 0.0f;
+	const float Ratio = ViewportClient->GetViewportHorizontalScrollBarRatio();
+	const float MaxOffset = (Ratio < 1.0f) ? 1.0f - Ratio : 0.0f;
 	InScrollOffsetFraction = FMath::Clamp(InScrollOffsetFraction, 0.0f, MaxOffset);
 
 	TextureViewportHorizontalScrollBar->SetState(InScrollOffsetFraction, Ratio);
 }
 
-
 EVisibility STextureEditorViewport::HandleHorizontalScrollBarVisibility( ) const
 {
 	if (ViewportClient->GetViewportHorizontalScrollBarRatio() < 1.0f)
@@ -182,17 +153,15 @@ EVisibility STextureEditorViewport::HandleHorizontalScrollBarVisibility( ) const
 	return EVisibility::Collapsed;
 }
 
-
 void STextureEditorViewport::HandleVerticalScrollBarScrolled( float InScrollOffsetFraction )
 {
-	float Ratio = ViewportClient->GetViewportVerticalScrollBarRatio();
-	float MaxOffset = (Ratio < 1.0f) ? 1.0f - Ratio : 0.0f;
+	const float Ratio = ViewportClient->GetViewportVerticalScrollBarRatio();
+	const float MaxOffset = (Ratio < 1.0f) ? 1.0f - Ratio : 0.0f;
 	InScrollOffsetFraction = FMath::Clamp(InScrollOffsetFraction, 0.0f, MaxOffset);
 
 	TextureViewportVerticalScrollBar->SetState(InScrollOffsetFraction, Ratio);
 }
 
-
 EVisibility STextureEditorViewport::HandleVerticalScrollBarVisibility( ) const
 {
 	if (ViewportClient->GetViewportVerticalScrollBarRatio() < 1.0f)
@@ -204,6 +173,4 @@ EVisibility STextureEditorViewport::HandleVerticalScrollBarVisibility( ) const
 }
 
 
-
-
 #undef LOCTEXT_NAMESPACE
diff --git a/Engine/Source/Editor/TextureEditor/Private/Widgets/STextureEditorViewport.h b/Engine/Source/Editor/TextureEditor/Private/Widgets/STextureEditorViewport.h
--- a/Engine/Source/Editor/TextureEditor/Private/Widgets/STextureEditorViewport.h
+++ b/Engine/Source/Editor/TextureEditor/Private/Widgets/STextureEditorViewport.h
@@ -24,8 +24,6 @@ public:
 	SLATE_BEGIN_ARGS(STextureEditorViewport) { }
 	SLATE_END_ARGS()
 
-public:
-
 	/**
 	 */
 	void AddReferencedObjects( FReferenceCollector& Collector );
@@ -83,13 +81,6 @@ private:
 	// Callback for getting the visibility of the horizontal scroll bar.
 	EVisibility HandleVerticalScrollBarVisibility( ) const;
 
-
-	// Checks if the texture being edited has a valid texture resource
-	bool HasValidTextureResource( ) const;
-
-private:
-
-
 	// Pointer back to the Texture editor tool that owns us.
 	TWeakPtr<ITextureEditorToolkit> ToolkitPtr;
 	
@@ -109,5 +100,5 @@ private:
 	TSharedPtr<SScrollBar> TextureViewportHorizontalScrollBar;
 
 	// Is rendering currently enabled? (disabled when reimporting a texture)
-	bool bIsRenderingEnabled;
+	bool bIsRenderingEnabled = true;
 };
