TextureEditorToolkit : nop comment about possibility to hide VirtualTextureStreaming property if r.VT.MenuRestricted

diff --git a/Engine/Source/Editor/TextureEditor/Private/TextureEditorToolkit.cpp b/Engine/Source/Editor/TextureEditor/Private/TextureEditorToolkit.cpp
--- a/Engine/Source/Editor/TextureEditor/Private/TextureEditorToolkit.cpp
+++ b/Engine/Source/Editor/TextureEditor/Private/TextureEditorToolkit.cpp
@@ -1161,6 +1161,11 @@ TSharedRef<SWidget> FTextureEditorToolkit::BuildTexturePropertiesWidget( )
 	FPropertyEditorModule& PropertyModule = FModuleManager::LoadModuleChecked<FPropertyEditorModule>("PropertyEditor");
 	TexturePropertiesWidget = PropertyModule.CreateDetailView(Args);
 	TexturePropertiesWidget->SetObject(Texture);
+	
+	// todo, maybe :
+	// check cvar "r.VT.MenuRestricted"
+	// if set, hide the "VirtualTextureStreaming" property
+	// TexturePropertiesWidget->SetIsPropertyVisibleDelegate(FIsPropertyVisible::CreateSP(this, &FTextureEditorToolkit::ShowTextureProperty));
 
 	return TexturePropertiesWidget.ToSharedRef();
 }
