TextureEditorToolkit : don't force final encoding on huge images

diff --git a/Engine/Source/Editor/TextureEditor/Private/TextureEditorToolkit.cpp b/Engine/Source/Editor/TextureEditor/Private/TextureEditorToolkit.cpp
--- a/Engine/Source/Editor/TextureEditor/Private/TextureEditorToolkit.cpp
+++ b/Engine/Source/Editor/TextureEditor/Private/TextureEditorToolkit.cpp
@@ -322,7 +322,10 @@ void FTextureEditorToolkit::InitTextureEditor( const EToolkitMode::Type Mode, co
 	// for the current editor session, but since it's already getting built this is
 	// fine. This doesn't need the normal Modify() / PreEditChange/PostEditChange incantations
 	// because it's transient, and any async build was completed above.
-	Texture->CompressFinal = true;
+	if ( Texture->Source.IsValid() && ( (int64)Texture->Source.GetSizeX() * Texture->Source.GetSizeY() <= 4096*4096  ) )
+	{
+		Texture->CompressFinal = true;
+	}
 	PostTextureRecode();
 
 	// @todo toolkit world centric editing
