Texture: add latest button next to OodleTextureSdkVersion to encourage its use

#rnx

diff --git a/Engine/Source/Editor/TextureEditor/Private/Customizations/TextureDetailsCustomization.cpp b/Engine/Source/Editor/TextureEditor/Private/Customizations/TextureDetailsCustomization.cpp
--- a/Engine/Source/Editor/TextureEditor/Private/Customizations/TextureDetailsCustomization.cpp
+++ b/Engine/Source/Editor/TextureEditor/Private/Customizations/TextureDetailsCustomization.cpp
@@ -10,6 +10,7 @@
 #include "IDetailPropertyRow.h"
 #include "DetailCategoryBuilder.h"
 #include "Widgets/Input/SNumericEntryBox.h"
+#include "Widgets/Input/SButton.h"
 
 #define LOCTEXT_NAMESPACE "FTextureDetails"
 
@@ -34,9 +35,58 @@ void FTextureDetails::CustomizeDetails(IDetailLayoutBuilder& DetailBuilder)
 	DetailBuilder.EditCategory("Adjustments");
 	DetailBuilder.EditCategory("File Path");
 
+	OodleTextureSdkVersionPropertyHandle = DetailBuilder.GetProperty(GET_MEMBER_NAME_CHECKED(UTexture, OodleTextureSdkVersion));
 	MaxTextureSizePropertyHandle = DetailBuilder.GetProperty(GET_MEMBER_NAME_CHECKED(UTexture, MaxTextureSize));
 	VirtualTextureStreamingPropertyHandle = DetailBuilder.GetProperty(GET_MEMBER_NAME_CHECKED(UTexture, VirtualTextureStreaming));
+		
+	if( OodleTextureSdkVersionPropertyHandle->IsValidHandle() )
+	{
+		IDetailCategoryBuilder& CompressionCategory = DetailBuilder.EditCategory("Compression");
+		IDetailPropertyRow& OodleTextureSdkVersionPropertyRow = CompressionCategory.AddProperty(GET_MEMBER_NAME_CHECKED(UTexture, OodleTextureSdkVersion));
+		TSharedPtr<SWidget> NameWidget;
+		TSharedPtr<SWidget> ValueWidget;
+		FDetailWidgetRow Row;
+		OodleTextureSdkVersionPropertyRow.GetDefaultWidgets(NameWidget, ValueWidget, Row);
 
+		const bool bShowChildren = true;
+		OodleTextureSdkVersionPropertyRow.CustomWidget(bShowChildren)
+			.NameContent()
+			.MinDesiredWidth(Row.NameWidget.MinWidth)
+			.MaxDesiredWidth(Row.NameWidget.MaxWidth)
+			[
+				NameWidget.ToSharedRef()
+			]
+			.ValueContent()
+			.MinDesiredWidth(Row.ValueWidget.MinWidth)
+			.MaxDesiredWidth(Row.ValueWidget.MaxWidth)
+			.VAlign(VAlign_Fill)
+			.HAlign(HAlign_Fill)
+			[
+				SNew(SHorizontalBox)
+				+ SHorizontalBox::Slot()
+				[
+					ValueWidget.ToSharedRef()
+				]
+				+ SHorizontalBox::Slot()
+				.HAlign(HAlign_Right)
+				.VAlign(VAlign_Center)
+				.AutoWidth()
+				[
+					SNew(SButton)
+					.OnClicked(this, &FTextureDetails::OnOodleTextureSdkVersionClicked)
+					.ContentPadding(FMargin(2))
+					.Content()
+					[
+						SNew(STextBlock)
+						.Justification(ETextJustify::Center)
+						.Font(IDetailLayoutBuilder::GetDetailFont())
+						.Text(LOCTEXT("OodleTextureSdkVersionLatest", "latest"))
+						.ToolTipText(LOCTEXT("OodleTextureSdkVersionLatestTooltip", "Update SDK Version to Latest"))
+					]
+				]
+			];
+	}
+	
 	// Customize MaxTextureSize
 	if( MaxTextureSizePropertyHandle->IsValidHandle() )
 	{
@@ -94,6 +144,20 @@ void FTextureDetails::CustomizeDetails(IDetailLayoutBuilder& DetailBuilder)
 	}
 }
 
+FReply FTextureDetails::OnOodleTextureSdkVersionClicked()
+{
+	UTexture* Texture = Cast<UTexture>(TextureBeingCustomized.Get());
+	if ( Texture == nullptr )
+	{
+		return FReply::Unhandled();
+	}
+
+	// true = do Pre/PostEditChange
+	Texture->UpdateOodleTextureSdkVersionToLatest(true);
+
+	return FReply::Handled();
+}
+
 /** @return The value or unset if properties with multiple values are viewed */
 TOptional<int32> FTextureDetails::OnGetMaxTextureSize() const
 {
diff --git a/Engine/Source/Editor/TextureEditor/Private/Customizations/TextureDetailsCustomization.h b/Engine/Source/Editor/TextureEditor/Private/Customizations/TextureDetailsCustomization.h
--- a/Engine/Source/Editor/TextureEditor/Private/Customizations/TextureDetailsCustomization.h
+++ b/Engine/Source/Editor/TextureEditor/Private/Customizations/TextureDetailsCustomization.h
@@ -7,6 +7,7 @@
 #include "UObject/WeakObjectPtr.h"
 #include "Types/SlateEnums.h"
 #include "IDetailCustomization.h"
+#include "Input/Reply.h"
 
 class IDetailLayoutBuilder;
 class IPropertyHandle;
@@ -30,7 +31,9 @@ private:
 	void OnMaxTextureSizeCommitted(int32 NewValue, ETextCommit::Type CommitInfo);
 	void OnBeginSliderMovement();
 	void OnEndSliderMovement(int32 NewValue);
+	FReply OnOodleTextureSdkVersionClicked();
 
+	TSharedPtr<IPropertyHandle> OodleTextureSdkVersionPropertyHandle;
 	TSharedPtr<IPropertyHandle> MaxTextureSizePropertyHandle;
 	TSharedPtr<IPropertyHandle> VirtualTextureStreamingPropertyHandle;
 	TWeakObjectPtr<UObject> TextureBeingCustomized;
diff --git a/Engine/Source/Runtime/Engine/Classes/Engine/Texture.h b/Engine/Source/Runtime/Engine/Classes/Engine/Texture.h
--- a/Engine/Source/Runtime/Engine/Classes/Engine/Texture.h
+++ b/Engine/Source/Runtime/Engine/Classes/Engine/Texture.h
@@ -1630,7 +1630,7 @@ public:
 	/* Change the Oodle Texture Sdk Version used to encode this texture to latest.
 	*  You should do this any time the texture is modified, such as on reimport, since the bits are changing anyway.
 	*/
-	ENGINE_API virtual void UpdateOodleTextureSdkVersionToLatest(void);
+	ENGINE_API virtual void UpdateOodleTextureSdkVersionToLatest(bool bDoPrePostEditChangeIfChanging = false);
 	
 	/** Get TextureFormatName with platform remaps and conditional prefix
 	 *   this is the entry point API for getting the final texture format name
diff --git a/Engine/Source/Runtime/Engine/Private/Texture.cpp b/Engine/Source/Runtime/Engine/Private/Texture.cpp
--- a/Engine/Source/Runtime/Engine/Private/Texture.cpp
+++ b/Engine/Source/Runtime/Engine/Private/Texture.cpp
@@ -472,11 +472,27 @@ bool UTexture::CanEditChange(const FProperty* InProperty) const
 	return true;
 }
 
-void UTexture::UpdateOodleTextureSdkVersionToLatest(void)
+void UTexture::UpdateOodleTextureSdkVersionToLatest(bool bDoPrePostEditChangeIfChanging)
 {
+	FName LatestVersion = CachedGetLatestOodleSdkVersion();
+	if ( OodleTextureSdkVersion == LatestVersion )
+	{
+		return;
+	}
+
+	if ( bDoPrePostEditChangeIfChanging )
+	{
+		PreEditChange(nullptr);
+	}
+
 	// OodleTextureSdkVersion = get latest sdk version
 	//	this needs to get the actual version number so it will be IO'd frozen (not just "latest")
-	OodleTextureSdkVersion = CachedGetLatestOodleSdkVersion();
+	OodleTextureSdkVersion = LatestVersion;
+	
+	if ( bDoPrePostEditChangeIfChanging )
+	{
+		PostEditChange();
+	}
 }
 
 
