More type casting fixes.

diff --git a/Engine/Source/Editor/TextureEditor/Private/Models/TextureEditorViewportClient.cpp b/Engine/Source/Editor/TextureEditor/Private/Models/TextureEditorViewportClient.cpp
--- a/Engine/Source/Editor/TextureEditor/Private/Models/TextureEditorViewportClient.cpp
+++ b/Engine/Source/Editor/TextureEditor/Private/Models/TextureEditorViewportClient.cpp
@@ -73,10 +73,10 @@ void FTextureEditorViewportClient::Draw(FViewport* Viewport, FCanvas* Canvas)
 	FVector2D ViewportSize = FVector2D(TextureEditorViewportPtr.Pin()->GetViewport()->GetSizeXY());
 	FVector2D ScrollBarPos = GetViewportScrollBarPositions();
 	int32 BorderSize = Settings.TextureBorderEnabled ? 1 : 0;
-	float YOffset = (Ratio.Y > 1.0f) ? ((ViewportSize.Y - (ViewportSize.Y / Ratio.Y)) * 0.5f) : 0;
-	int32 YPos = FMath::RoundToInt(YOffset - ScrollBarPos.Y + BorderSize);
-	float XOffset = (Ratio.X > 1.0f) ? ((ViewportSize.X - (ViewportSize.X / Ratio.X)) * 0.5f) : 0;
-	int32 XPos = FMath::RoundToInt(XOffset - ScrollBarPos.X + BorderSize);
+	float YOffset = static_cast<float>((Ratio.Y > 1.0) ? ((ViewportSize.Y - (ViewportSize.Y / Ratio.Y)) * 0.5) : 0);
+	int32 YPos = (int32)FMath::Clamp(FMath::RoundToInt(YOffset - ScrollBarPos.Y + BorderSize), TNumericLimits<int32>::Min(), TNumericLimits<int32>::Max());
+	float XOffset = static_cast<float>((Ratio.X > 1.0) ? ((ViewportSize.X - (ViewportSize.X / Ratio.X)) * 0.5) : 0);
+	int32 XPos = (int32)FMath::Clamp(FMath::RoundToInt(XOffset - ScrollBarPos.X + BorderSize), TNumericLimits<int32>::Min(), TNumericLimits<int32>::Max());
 	
 	UpdateScrollBars();
 
@@ -222,7 +222,7 @@ void FTextureEditorViewportClient::Draw(FViewport* Viewport, FCanvas* Canvas)
 		if (Settings.TextureBorderEnabled)
 		{
 			FCanvasBoxItem BoxItem(FVector2D(XPos - (BorderSize - 1) * 0.5f, YPos - (BorderSize - 1) * 0.5f), FVector2D(Width + BorderSize, Height + BorderSize));
-			BoxItem.LineThickness = BorderSize;
+			BoxItem.LineThickness = (float)BorderSize;
 			BoxItem.SetColor( Settings.TextureBorderColor );
 			Canvas->DrawItem( BoxItem );
 		}
@@ -244,7 +244,7 @@ void FTextureEditorViewportClient::Draw(FViewport* Viewport, FCanvas* Canvas)
 				IAllocatedVirtualTexture* AllocatedVT = VTResource->AcquireAllocatedVT();
 
 				IRendererModule& RenderModule = GetRendererModule();
-				RenderModule.RequestVirtualTextureTilesForRegion(AllocatedVT, ScreenSpaceSize, ViewportPositon, ViewportSize, UV0, UV1, MipLevel);
+				RenderModule.RequestVirtualTextureTilesForRegion(AllocatedVT, ScreenSpaceSize, ViewportPositon, ViewportSize, UV0, UV1, (int32)MipLevel);
 				RenderModule.LoadPendingVirtualTextureTiles(RHICmdList, InFeatureLevel);
 			});
 		}
@@ -267,7 +267,7 @@ void FTextureEditorViewportClient::Draw(FViewport* Viewport, FCanvas* Canvas)
 			const int32 LineHeight = FMath::TruncToInt(ErrorFont->GetMaxCharHeight());
 			const FText Message = NSLOCTEXT("TextureEditor", "InvalidVirtualTextureMipDisplay", "Displaying a virtual texture on a mip level that is larger than the physical cache. Rendering will probably be invalid!");
 			const uint32 MessageWidth = ErrorFont->GetStringSize(*Message.ToString());
-			const uint32 Xpos = (ViewportSize.X - MessageWidth) / 2;
+			const uint32 Xpos = (uint32)((ViewportSize.X - MessageWidth) / 2);
 			Canvas->DrawShadowedText(Xpos, LineHeight*1.5,
 				Message,
 				ErrorFont, FLinearColor::Red);
@@ -435,7 +435,7 @@ float FTextureEditorViewportClient::GetViewportVerticalScrollBarRatio() const
 		int32 Depth, ArraySize;
 		TextureEditorPtr.Pin()->CalculateTextureDimensions(Width, Height, Depth, ArraySize, true);
 
-		WidgetHeight = TextureEditorViewportPtr.Pin()->GetViewport()->GetSizeXY().Y;
+		WidgetHeight = (float)(TextureEditorViewportPtr.Pin()->GetViewport()->GetSizeXY().Y);
 	}
 
 	return WidgetHeight / Height;
@@ -452,7 +452,7 @@ float FTextureEditorViewportClient::GetViewportHorizontalScrollBarRatio() const
 		int32 Depth, ArraySize;
 		TextureEditorPtr.Pin()->CalculateTextureDimensions(Width, Height, Depth, ArraySize, true);
 
-		WidgetWidth = TextureEditorViewportPtr.Pin()->GetViewport()->GetSizeXY().X;
+		WidgetWidth = (float)(TextureEditorViewportPtr.Pin()->GetViewport()->GetSizeXY().X);
 	}
 
 	return WidgetWidth / Width;
