#UE Fix issues with editing multiple textures at once in the property matrix

diff --git a/Engine/Source/Editor/DetailCustomizations/Private/PerPlatformPropertyCustomization.cpp b/Engine/Source/Editor/DetailCustomizations/Private/PerPlatformPropertyCustomization.cpp
--- a/Engine/Source/Editor/DetailCustomizations/Private/PerPlatformPropertyCustomization.cpp
+++ b/Engine/Source/Editor/DetailCustomizations/Private/PerPlatformPropertyCustomization.cpp
@@ -104,6 +104,7 @@ TSharedRef<SWidget> FPerPlatformPropertyCustomization<PerPlatformType>::GetWidge
 	}
 
 	// Push down struct metadata to per-platform properties
+	if (EditProperty.IsValid())
 	{
 		// First get the source map
 		const TMap<FName, FString>* SourceMap = StructPropertyHandle->GetMetaDataProperty()->GetMetaDataMap();
@@ -119,15 +120,12 @@ TSharedRef<SWidget> FPerPlatformPropertyCustomization<PerPlatformType>::GetWidge
 		{
 			EditProperty->SetInstanceMetaData(*It.Key.ToString(), *It.Value);
 		}
-	}
 
-	if (ensure(EditProperty.IsValid()))
-	{
 		if (EditProperty->GetProperty()->IsA<FStructProperty>())
 		{
 			return StructBuilder.GenerateStructValueWidget(EditProperty->AsShared());
 		}
-		
+
 		return EditProperty->CreatePropertyValueWidget(false);
 	}
 	
diff --git a/Engine/Source/Editor/TextureEditor/Private/Customizations/TextureDetailsCustomization.cpp b/Engine/Source/Editor/TextureEditor/Private/Customizations/TextureDetailsCustomization.cpp
--- a/Engine/Source/Editor/TextureEditor/Private/Customizations/TextureDetailsCustomization.cpp
+++ b/Engine/Source/Editor/TextureEditor/Private/Customizations/TextureDetailsCustomization.cpp
@@ -22,12 +22,7 @@ TSharedRef<IDetailCustomization> FTextureDetails::MakeInstance()
 
 void FTextureDetails::CustomizeDetails(IDetailLayoutBuilder& DetailBuilder)
 {
-	TArray< TWeakObjectPtr<UObject> > ObjectsBeingCustomized;
-	DetailBuilder.GetObjectsBeingCustomized(ObjectsBeingCustomized);
-	if (ensure(ObjectsBeingCustomized.Num() == 1))
-	{
-		TextureBeingCustomized = ObjectsBeingCustomized[0];
-	}
+	DetailBuilder.GetObjectsBeingCustomized(TexturesBeingCustomized);
 
 	DetailBuilder.EditCategory("LevelOfDetail");
 	DetailBuilder.EditCategory("Compression");
@@ -88,7 +83,7 @@ void FTextureDetails::CustomizeDetails(IDetailLayoutBuilder& DetailBuilder)
 	}
 	
 	// Customize MaxTextureSize
-	if( MaxTextureSizePropertyHandle->IsValidHandle() )
+	if( MaxTextureSizePropertyHandle->IsValidHandle() && TexturesBeingCustomized.Num() == 1)
 	{
 		IDetailCategoryBuilder& CompressionCategory = DetailBuilder.EditCategory("Compression");
 		IDetailPropertyRow& MaxTextureSizePropertyRow = CompressionCategory.AddProperty(GET_MEMBER_NAME_CHECKED(UTexture, MaxTextureSize));
@@ -99,7 +94,7 @@ void FTextureDetails::CustomizeDetails(IDetailLayoutBuilder& DetailBuilder)
 
 		int32 MaxTextureSize = 2048;
 
-		if (UTexture* Texture = Cast<UTexture>(TextureBeingCustomized.Get()))
+		if (UTexture* Texture = Cast<UTexture>(TexturesBeingCustomized[0].Get()))
 		{
 			// GetMaximumDimension is for current RHI and texture type
 			MaxTextureSize = Texture->GetMaximumDimension();
@@ -146,15 +141,15 @@ void FTextureDetails::CustomizeDetails(IDetailLayoutBuilder& DetailBuilder)
 
 FReply FTextureDetails::OnOodleTextureSdkVersionClicked()
 {
-	UTexture* Texture = Cast<UTexture>(TextureBeingCustomized.Get());
-	if ( Texture == nullptr )
+	for (const TWeakObjectPtr<UObject>& WeakTexture : TexturesBeingCustomized)
 	{
-		return FReply::Unhandled();
+		if (UTexture* Texture = Cast<UTexture>(WeakTexture.Get()))
+		{
+			// true = do Pre/PostEditChange
+			Texture->UpdateOodleTextureSdkVersionToLatest(true);
+		}
 	}
 
-	// true = do Pre/PostEditChange
-	Texture->UpdateOodleTextureSdkVersionToLatest(true);
-
 	return FReply::Handled();
 }
 
diff --git a/Engine/Source/Editor/TextureEditor/Private/Customizations/TextureDetailsCustomization.h b/Engine/Source/Editor/TextureEditor/Private/Customizations/TextureDetailsCustomization.h
--- a/Engine/Source/Editor/TextureEditor/Private/Customizations/TextureDetailsCustomization.h
+++ b/Engine/Source/Editor/TextureEditor/Private/Customizations/TextureDetailsCustomization.h
@@ -36,7 +36,7 @@ private:
 	TSharedPtr<IPropertyHandle> OodleTextureSdkVersionPropertyHandle;
 	TSharedPtr<IPropertyHandle> MaxTextureSizePropertyHandle;
 	TSharedPtr<IPropertyHandle> VirtualTextureStreamingPropertyHandle;
-	TWeakObjectPtr<UObject> TextureBeingCustomized;
+	TArray<TWeakObjectPtr<UObject>> TexturesBeingCustomized;
 	bool bIsUsingSlider;
 };
 
