Fix for dynamically sources textures going black when opening in the texture editor.

diff --git a/Engine/Source/Editor/TextureEditor/Private/TextureEditorToolkit.cpp b/Engine/Source/Editor/TextureEditor/Private/TextureEditorToolkit.cpp
--- a/Engine/Source/Editor/TextureEditor/Private/TextureEditorToolkit.cpp
+++ b/Engine/Source/Editor/TextureEditor/Private/TextureEditorToolkit.cpp
@@ -336,7 +336,17 @@ void FTextureEditorToolkit::InitTextureEditor( const EToolkitMode::Type Mode, co
 	{
 		Texture->CompressFinal = true;
 	}
-	PostTextureRecode();
+
+	// We don't want to post recodes for render targets because that clears them to black and
+	// we don't care about CompressFinal for them anyway as they aren't encoded. While we are
+	// here, don't bother with other dynamic textures as well.
+	ETextureClass TextureClass = Texture->GetTextureClass();
+	if (TextureClass != ETextureClass::RenderTarget &&
+		TextureClass != ETextureClass::Other2DNoSource &&
+		TextureClass != ETextureClass::TwoDDynamic)
+	{
+		PostTextureRecode();
+	}
 
 	// @todo toolkit world centric editing
 	/*if(IsWorldCentricAssetEditor())
