TextureEditorToolkit fix crash opening texture editor while texture is building

#rnx

diff --git a/Engine/Source/Editor/TextureEditor/Private/TextureEditorToolkit.cpp b/Engine/Source/Editor/TextureEditor/Private/TextureEditorToolkit.cpp
--- a/Engine/Source/Editor/TextureEditor/Private/TextureEditorToolkit.cpp
+++ b/Engine/Source/Editor/TextureEditor/Private/TextureEditorToolkit.cpp
@@ -2546,7 +2546,7 @@ void FTextureEditorToolkit::HandleReimportManagerPreReimport( UObject* InObject
 
 void FTextureEditorToolkit::HandleAssetPostImport(UFactory* InFactory, UObject* InObject)
 {
-	if (Cast<UTexture>(InObject) != nullptr && InObject == Texture)
+	if (Cast<UTexture>(InObject) != nullptr && InObject == Texture && TexturePropertiesWidget != nullptr )
 	{
 		// Refresh this object within the details panel
 		TexturePropertiesWidget->SetObject(InObject);
