Copy relevant warnings/errors to the texture editor window.

diff --git a/Engine/Source/Editor/TextureEditor/Private/Models/TextureEditorViewportClient.cpp b/Engine/Source/Editor/TextureEditor/Private/Models/TextureEditorViewportClient.cpp
--- a/Engine/Source/Editor/TextureEditor/Private/Models/TextureEditorViewportClient.cpp
+++ b/Engine/Source/Editor/TextureEditor/Private/Models/TextureEditorViewportClient.cpp
@@ -35,6 +35,79 @@ static TAutoConsoleVariable<int32> CVarEnableVTFeedback(
 	ECVF_RenderThreadSafe
 );
 
+struct FTextureErrorLogger : public FOutputDevice
+{
+	UTexture* TextureToMonitor = nullptr;
+	TArray<TPair<bool, FString>> RelevantLogLines;
+	bool bCurrentlyCapturing = false;
+
+	FTextureErrorLogger(UTexture* InTextureToMonitor)
+	{
+		TextureToMonitor = InTextureToMonitor;
+		GLog->AddOutputDevice(this);
+	}
+	~FTextureErrorLogger()
+	{
+		GLog->RemoveOutputDevice(this);
+	}
+	FTextureErrorLogger(FTextureErrorLogger&&) = delete;
+	FTextureErrorLogger(FTextureErrorLogger&) = delete;
+	FTextureErrorLogger& operator=(FTextureErrorLogger&) = delete;
+	FTextureErrorLogger& operator=(FTextureErrorLogger&&) = delete;
+
+
+	virtual void Serialize(const TCHAR* V, ELogVerbosity::Type Verbosity, const class FName& Category) override
+	{
+		if (TextureToMonitor == nullptr)
+		{
+			return;
+		}
+
+		//
+		// Error messages don't reliably put the texture name in the messages, and we don't necessarily know
+		// that the error will come from a texture category if it's something like bulk data. However, we do
+		// ~generally~ expect that if we have a texture editor open, the only texture that's getting built is
+		// the one we are editing. So we just capture all errors/warning between the time we see "building texture"
+		// for ourselves and when the texture async build is complete.
+		if (bCurrentlyCapturing)
+		{
+			if (TextureToMonitor->IsAsyncCacheComplete())
+			{
+				bCurrentlyCapturing = false;
+				return;
+			}
+
+			// Add any errors or warnings to the list
+			if (Verbosity == ELogVerbosity::Error)
+			{
+				RelevantLogLines.Add(TPair<bool, FString>(true, FString(V)));
+			}
+			else if (Verbosity == ELogVerbosity::Warning)
+			{
+				RelevantLogLines.Add(TPair<bool, FString>(false, FString(V)));
+			}
+			if (RelevantLogLines.Num() == 10)
+			{
+				RelevantLogLines.Add(TPair<bool, FString>(false, TEXT("Too much to show: check Output Log")));
+				bCurrentlyCapturing = false;
+			}
+			return;
+		}
+
+		// Here we aren't capturing yet.
+		// See if the string relates to us.
+		if (FCString::Stristr(V, *TextureToMonitor->GetName()))
+		{
+			// If it's "Building textures" then we started a new build and need to empty our list.
+			if (FCString::Stristr(V, TEXT("Building textures")))
+			{
+				RelevantLogLines.Empty();
+				bCurrentlyCapturing = true;
+			}
+		}
+	}
+};
+
 /* FTextureEditorViewportClient structors
  *****************************************************************************/
 
@@ -46,6 +119,8 @@ FTextureEditorViewportClient::FTextureEditorViewportClient( TWeakPtr<ITextureEdi
 	check(TextureEditorPtr.IsValid() && TextureEditorViewportPtr.IsValid());
 
 	ModifyCheckerboardTextureColors();
+
+	TextureConsoleCapture = MakeUnique<FTextureErrorLogger>(InTextureEditor.Pin()->GetTexture());
 }
 
 
@@ -269,6 +344,12 @@ void FTextureEditorViewportClient::Draw(FViewport* Viewport, FCanvas* Canvas)
 		}
 	}
 
+	UFont* ReportingFont = GEngine->GetLargeFont();
+	const int32 ReportingLineHeight = FMath::CeilToInt(ReportingFont->GetMaxCharHeight()) + 2; // 2 for line spacing
+	const int32 ReportingLineX = 8;
+	int32 ReportingLineY = 8;
+
+
 	// If we are requesting an explicit mip level of a VT asset, test to see if we can even display it properly and warn about it
 	if (bIsVirtualTexture && MipLevel >= 0.f)
 	{
@@ -282,16 +363,17 @@ void FTextureEditorViewportClient::Draw(FViewport* Viewport, FCanvas* Canvas)
 
 		if (NumPixels >= NumPhysicalPixels)
 		{
-			UFont* ErrorFont = GEngine->GetLargeFont();
-			const int32 LineHeight = FMath::TruncToInt(ErrorFont->GetMaxCharHeight());
 			const FText Message = NSLOCTEXT("TextureEditor", "InvalidVirtualTextureMipDisplay", "Displaying a virtual texture on a mip level that is larger than the physical cache. Rendering will probably be invalid!");
-			const uint32 MessageWidth = ErrorFont->GetStringSize(*Message.ToString());
-			const uint32 Xpos = (uint32)((ViewportSize.X - MessageWidth) / 2);
-			Canvas->DrawShadowedText(Xpos, LineHeight*1.5,
-				Message,
-				ErrorFont, FLinearColor::Red);
+			Canvas->DrawShadowedText(ReportingLineX, ReportingLineY, Message, ReportingFont, FLinearColor::Red);
+			ReportingLineY += ReportingLineHeight;
 		}
-		
+	}
+
+	// Print any warnings/errors that we saw in the output log.
+	for (TPair<bool, FString>& ReportedLine : TextureConsoleCapture->RelevantLogLines)
+	{
+		Canvas->DrawShadowedText(ReportingLineX, ReportingLineY, FText::FromString(ReportedLine.Value), GEngine->GetLargeFont(), ReportedLine.Key ? FLinearColor::Red : FLinearColor::Yellow);
+		ReportingLineY += ReportingLineHeight;
 	}
 }
 
diff --git a/Engine/Source/Editor/TextureEditor/Private/Models/TextureEditorViewportClient.h b/Engine/Source/Editor/TextureEditor/Private/Models/TextureEditorViewportClient.h
--- a/Engine/Source/Editor/TextureEditor/Private/Models/TextureEditorViewportClient.h
+++ b/Engine/Source/Editor/TextureEditor/Private/Models/TextureEditorViewportClient.h
@@ -69,4 +69,7 @@ private:
 
 	/** Checkerboard texture */
 	TObjectPtr<UTexture2D> CheckerboardTexture;
+
+	/** Output device to monitor the output log for anything relevant to us.*/
+	TUniquePtr<struct FTextureErrorLogger> TextureConsoleCapture;
 };
