Make it clear when the user is viewing unencoded data due to a deferred encoding in the texture editor.

diff --git a/Engine/Source/Editor/TextureEditor/Private/Models/TextureEditorViewportClient.cpp b/Engine/Source/Editor/TextureEditor/Private/Models/TextureEditorViewportClient.cpp
--- a/Engine/Source/Editor/TextureEditor/Private/Models/TextureEditorViewportClient.cpp
+++ b/Engine/Source/Editor/TextureEditor/Private/Models/TextureEditorViewportClient.cpp
@@ -19,6 +19,7 @@
 #include "Slate/SceneViewport.h"
 #include "Texture2DPreview.h"
 #include "VolumeTexturePreview.h"
+#include "TextureEncodingSettings.h"
 #include "TextureEditorSettings.h"
 #include "Widgets/STextureEditorViewport.h"
 #include "CanvasTypes.h"
@@ -369,6 +370,107 @@ void FTextureEditorViewportClient::Draw(FViewport* Viewport, FCanvas* Canvas)
 		}
 	}
 
+	// If we have compression deferred, make it clear they are viewing unencoded data.
+	if (Texture->DeferCompression)
+	{
+		const FText Message = NSLOCTEXT("TextureEditor", "CompressionDeferred", "Compression Deferred: Viewing unencoded data!");
+		Canvas->DrawShadowedText(ReportingLineX, ReportingLineY, Message, ReportingFont, FLinearColor::Yellow);
+		ReportingLineY += ReportingLineHeight;
+	}
+	else
+	{
+		// Check if we are viewing an encoding that isn't Final.
+		FTexturePlatformData** PlatformDataPtr = Texture->GetRunningPlatformData();
+		
+		if (PlatformDataPtr &&
+			PlatformDataPtr[0] && // Can be null if we haven't had a chance to call CachePlatformData on the texture (brand new)
+			PlatformDataPtr[0]->ResultMetadata.bIsValid)
+		{
+			FResolvedTextureEncodingSettings const& EncodeSettings = FResolvedTextureEncodingSettings::Get();
+			bool bEncodingDiffers = (EncodeSettings.Project.bFastUsesRDO != EncodeSettings.Project.bFinalUsesRDO ||
+				EncodeSettings.Project.FastEffortLevel != EncodeSettings.Project.FinalEffortLevel ||
+				EncodeSettings.Project.FastRDOLambda != EncodeSettings.Project.FinalRDOLambda ||
+				EncodeSettings.Project.FastUniversalTiling != EncodeSettings.Project.FastUniversalTiling);
+
+			if (PlatformDataPtr[0]->ResultMetadata.bWasEditorCustomEncoding)
+			{
+				const FText LeadInText = NSLOCTEXT("TextureEditor", "ViewingCustom", "Viewing custom encoding");
+				Canvas->DrawShadowedText(ReportingLineX, ReportingLineY, LeadInText, ReportingFont, FLinearColor::Yellow);
+				ReportingLineY += ReportingLineHeight;
+			}
+			else if (bEncodingDiffers &&
+				PlatformDataPtr[0]->ResultMetadata.bSupportsEncodeSpeed &&
+				PlatformDataPtr[0]->ResultMetadata.EncodeSpeed != (uint8)ETextureEncodeSpeed::Final)
+			{
+				// We aren't final - which might not matter if they encode the same way, so just show the differences from final.
+				int32 CurrentX = ReportingLineX;
+
+				auto DrawComma = [&CurrentX, &ReportingFont, &Canvas, &ReportingLineY]()
+				{
+					const TCHAR* Comma = TEXT(", ");
+					int32 CommaWidth = ReportingFont->GetStringSize(Comma);
+					Canvas->DrawShadowedString(CurrentX, ReportingLineY, Comma, ReportingFont, FLinearColor::Yellow);
+					CurrentX += CommaWidth;
+				};
+
+				const FText LeadInText = NSLOCTEXT("TextureEditor", "EncodingDifference", "Viewing non-shipping encoding, differences are: ");
+				Canvas->DrawShadowedText(CurrentX, ReportingLineY, LeadInText, ReportingFont, FLinearColor::Yellow);
+				CurrentX += ReportingFont->GetStringSize(*LeadInText.ToString()); // afaict you always need to ToString to measure the text.
+
+				const FText HelpText = NSLOCTEXT("TextureEditor", "ShowFinal", "Check \"Editor Show Final Encoding\" to see shipping encoding.");
+				Canvas->DrawShadowedText(ReportingLineX, ReportingLineY + ReportingLineHeight, HelpText, ReportingFont, FLinearColor::Yellow);
+
+				bool bNeedComma = false;
+				if (EncodeSettings.Project.bFastUsesRDO != EncodeSettings.Project.bFinalUsesRDO)
+				{
+					const FText RDOText = NSLOCTEXT("TextureEditor", "RDODifference", "RDO On/Off");
+					Canvas->DrawShadowedText(CurrentX, ReportingLineY, RDOText, ReportingFont, FLinearColor::Yellow);
+					CurrentX += ReportingFont->GetStringSize(*RDOText.ToString());
+					bNeedComma = true;
+				}
+				else if (EncodeSettings.Project.bFinalUsesRDO)
+				{
+					// Some stuff only matters if RDO is on
+					if (EncodeSettings.Project.FastRDOLambda != EncodeSettings.Project.FinalRDOLambda)
+					{
+						if (bNeedComma) 
+						{
+							DrawComma();
+						}
+						const FText RDOText = NSLOCTEXT("TextureEditor", "RDOLambdaDifference", "RDO Lambda");
+						Canvas->DrawShadowedText(CurrentX, ReportingLineY, RDOText, ReportingFont, FLinearColor::Yellow);
+						CurrentX += ReportingFont->GetStringSize(*RDOText.ToString());
+						bNeedComma = true;
+					}
+					if (EncodeSettings.Project.FastUniversalTiling != EncodeSettings.Project.FinalUniversalTiling)
+					{
+						if (bNeedComma)
+						{
+							DrawComma();
+						}
+						const FText EffortText = NSLOCTEXT("TextureEditor", "UTDifference", "RDO Universal Tiling");
+						Canvas->DrawShadowedText(CurrentX, ReportingLineY, EffortText, ReportingFont, FLinearColor::Yellow);
+						CurrentX += ReportingFont->GetStringSize(*EffortText.ToString());;
+						bNeedComma = true;
+					}
+				}
+				if (EncodeSettings.Project.FastEffortLevel != EncodeSettings.Project.FinalEffortLevel)
+				{
+					if (bNeedComma)
+					{
+						DrawComma();
+					}
+					const FText EffortText = NSLOCTEXT("TextureEditor", "EffortDifference", "Encode Effort");
+					Canvas->DrawShadowedText(CurrentX, ReportingLineY, EffortText, ReportingFont, FLinearColor::Yellow);
+					CurrentX += ReportingFont->GetStringSize(*EffortText.ToString());;
+					bNeedComma = true;
+				}
+
+				ReportingLineY += 2*ReportingLineHeight;
+			} // end if difference exists
+		} // end if valid result metadata
+	} // end if not deferring
+
 	// Print any warnings/errors that we saw in the output log.
 	for (TPair<bool, FString>& ReportedLine : TextureConsoleCapture->RelevantLogLines)
 	{
