PVS fix for silly mistake

diff --git a/Engine/Source/Editor/TextureEditor/Private/Models/TextureEditorViewportClient.cpp b/Engine/Source/Editor/TextureEditor/Private/Models/TextureEditorViewportClient.cpp
--- a/Engine/Source/Editor/TextureEditor/Private/Models/TextureEditorViewportClient.cpp
+++ b/Engine/Source/Editor/TextureEditor/Private/Models/TextureEditorViewportClient.cpp
@@ -390,7 +390,7 @@ void FTextureEditorViewportClient::Draw(FViewport* Viewport, FCanvas* Canvas)
 			bool bEncodingDiffers = (EncodeSettings.Project.bFastUsesRDO != EncodeSettings.Project.bFinalUsesRDO ||
 				EncodeSettings.Project.FastEffortLevel != EncodeSettings.Project.FinalEffortLevel ||
 				EncodeSettings.Project.FastRDOLambda != EncodeSettings.Project.FinalRDOLambda ||
-				EncodeSettings.Project.FastUniversalTiling != EncodeSettings.Project.FastUniversalTiling);
+				EncodeSettings.Project.FastUniversalTiling != EncodeSettings.Project.FinalUniversalTiling);
 
 			if (PlatformDataPtr[0]->ResultMetadata.bWasEditorCustomEncoding)
 			{
