Minor fixes for the in-viewer texture build error reporting.

diff --git a/Engine/Source/Editor/TextureEditor/Private/Models/TextureEditorViewportClient.cpp b/Engine/Source/Editor/TextureEditor/Private/Models/TextureEditorViewportClient.cpp
--- a/Engine/Source/Editor/TextureEditor/Private/Models/TextureEditorViewportClient.cpp
+++ b/Engine/Source/Editor/TextureEditor/Private/Models/TextureEditorViewportClient.cpp
@@ -40,6 +40,8 @@ struct FTextureErrorLogger : public FOutputDevice
 {
 	UTexture* TextureToMonitor = nullptr;
 	TArray<TPair<bool, FString>> RelevantLogLines;
+
+	double CaptureUntilTime = 0;
 	bool bCurrentlyCapturing = false;
 
 	FTextureErrorLogger(UTexture* InTextureToMonitor)
@@ -72,10 +74,22 @@ struct FTextureErrorLogger : public FOutputDevice
 		// for ourselves and when the texture async build is complete.
 		if (bCurrentlyCapturing)
 		{
-			if (TextureToMonitor->IsAsyncCacheComplete())
+			// There's a race condition here because the log messages take a while to get to us, so an
+			// async texture build can finish well before we see any messages about it, so we keep watching for a second.
+			if (CaptureUntilTime)
 			{
-				bCurrentlyCapturing = false;
-				return;
+				double CurrentTime = FPlatformTime::Seconds();
+				if (CurrentTime > CaptureUntilTime)
+				{
+					bCurrentlyCapturing = false;
+					CaptureUntilTime = 0;
+					return;
+				}
+			}
+			else if (TextureToMonitor->IsAsyncCacheComplete())
+			{
+				// Logs might not have gotten to us, so we keep capturing for another second.
+				CaptureUntilTime = FPlatformTime::Seconds() + 1;
 			}
 
 			// Add any errors or warnings to the list
@@ -91,6 +105,7 @@ struct FTextureErrorLogger : public FOutputDevice
 			{
 				RelevantLogLines.Add(TPair<bool, FString>(false, TEXT("Too much to show: check Output Log")));
 				bCurrentlyCapturing = false;
+				CaptureUntilTime = 0;
 			}
 			return;
 		}
@@ -100,11 +115,17 @@ struct FTextureErrorLogger : public FOutputDevice
 		if (FCString::Stristr(V, *TextureToMonitor->GetName()))
 		{
 			// If it's "Building textures" then we started a new build and need to empty our list.
+			// For shared linear we also might only get a "Tiling" message... but we don't want to clear
+			// on that one since we'll likely get it after the Building textures message
 			if (FCString::Stristr(V, TEXT("Building textures")))
 			{
 				RelevantLogLines.Empty();
 				bCurrentlyCapturing = true;
 			}
+			else if (FCString::Stristr(V, TEXT("Tiling")))
+			{
+				bCurrentlyCapturing = true;
+			}
 		}
 	}
 };
