Fix waiting on async texture build in texture viewer.

diff --git a/Engine/Source/Editor/TextureEditor/Private/TextureEditorToolkit.cpp b/Engine/Source/Editor/TextureEditor/Private/TextureEditorToolkit.cpp
--- a/Engine/Source/Editor/TextureEditor/Private/TextureEditorToolkit.cpp
+++ b/Engine/Source/Editor/TextureEditor/Private/TextureEditorToolkit.cpp
@@ -371,12 +371,17 @@ void FTextureEditorToolkit::CalculateTextureDimensions(int32& OutWidth, int32& O
 	
 	if (UTexture2D* Texture2D = Cast<UTexture2D>(Texture))
 	{
-		if (UTexture2D* CpuTexture = Texture2D->GetCPUCopyTexture())
+		// GetCPUCopyTexture waits on the build if there is one.
+		if (Texture2D->Availability == ETextureAvailability::CPU &&
+			Texture2D->IsAsyncCacheComplete())
 		{
-			OutWidth = CpuTexture->GetSurfaceWidth();
-			OutHeight = CpuTexture->GetSurfaceHeight();
-			OutDepth = 1;
-			OutArraySize = 0;
+			if (UTexture2D* CpuTexture = Texture2D->GetCPUCopyTexture())
+			{
+				OutWidth = CpuTexture->GetSurfaceWidth();
+				OutHeight = CpuTexture->GetSurfaceHeight();
+				OutDepth = 1;
+				OutArraySize = 0;
+			}
 		}
 	}
 
@@ -869,7 +874,10 @@ void FTextureEditorToolkit::PopulateQuickInfo( )
 		MaxInGameText->SetText(FText::Format( NSLOCTEXT("TextureEditor", "QuickInfo_MaxInGame_2x", "Max In-Game: {0}x{1}{2}"), FText::AsNumber(MaxInGameWidth, &FormatOptions), FText::AsNumber(MaxInGameHeight, &FormatOptions), InGameCubemapInfo));
 	}
 
-	if (Texture2D)
+	// GetCPUCopyTexture waits on the build if there is one.
+	if (Texture2D &&
+		Texture2D->Availability == ETextureAvailability::CPU &&
+		Texture2D->IsAsyncCacheComplete())
 	{
 		if (UTexture2D* CpuTexture = Texture2D->GetCPUCopyTexture(); CpuTexture)
 		{
