TextureEditorToolkit : fix CompressFinal bool not getting turned off if it was once turned on

diff --git a/Engine/Source/Editor/TextureEditor/Private/TextureEditorToolkit.cpp b/Engine/Source/Editor/TextureEditor/Private/TextureEditorToolkit.cpp
--- a/Engine/Source/Editor/TextureEditor/Private/TextureEditorToolkit.cpp
+++ b/Engine/Source/Editor/TextureEditor/Private/TextureEditorToolkit.cpp
@@ -331,10 +331,7 @@ void FTextureEditorToolkit::InitTextureEditor( const EToolkitMode::Type Mode, co
 	// for the current editor session, but since it's already getting built this is
 	// fine. This doesn't need the normal Modify() / PreEditChange/PostEditChange incantations
 	// because it's transient, and any async build was completed above.
-	if ( Texture->Source.IsValid() && ( Texture->Source.GetTotalTopMipPixelCount() <= (int64)4096*4096  ) )
-	{
-		Texture->CompressFinal = true;
-	}
+	Texture->CompressFinal = ( Texture->Source.IsValid() && ( Texture->Source.GetTotalTopMipPixelCount() <= (int64)4096*4096  ) );
 
 	// We don't want to post recodes for render targets because that clears them to black and
 	// we don't care about CompressFinal for them anyway as they aren't encoded. While we are
