TextureEditor : make sure to restore temporary editor properties in texture on exit
fix render targets
clean up identification of textures that should be recoded

diff --git a/Engine/Source/Editor/TextureEditor/Private/TextureEditorModule.cpp b/Engine/Source/Editor/TextureEditor/Private/TextureEditorModule.cpp
--- a/Engine/Source/Editor/TextureEditor/Private/TextureEditorModule.cpp
+++ b/Engine/Source/Editor/TextureEditor/Private/TextureEditorModule.cpp
@@ -37,13 +37,6 @@ public:
 		TSharedRef<FTextureEditorToolkit> NewTextureEditor(new FTextureEditorToolkit());
 		NewTextureEditor->InitTextureEditor(Mode, InitToolkitHost, Texture);
 
-		// OpenAssetEditor should ensure that we never have two toolkits open for the same asset!
-		check(Texture->TextureEditorCustomEncoding.IsValid() == false);
-
-		// We save this as a separate object so that the engine can reference the type
-		// without needing the texture editor module.
-		Texture->TextureEditorCustomEncoding = NewTextureEditor->CustomEncoding;
-
 		return NewTextureEditor;
 	}
 
diff --git a/Engine/Source/Editor/TextureEditor/Private/TextureEditorToolkit.cpp b/Engine/Source/Editor/TextureEditor/Private/TextureEditorToolkit.cpp
--- a/Engine/Source/Editor/TextureEditor/Private/TextureEditorToolkit.cpp
+++ b/Engine/Source/Editor/TextureEditor/Private/TextureEditorToolkit.cpp
@@ -74,7 +74,7 @@ const FName FTextureEditorToolkit::OodleTabId(TEXT("TextureEditor_Oodle"));
 
 UNREALED_API void GetBestFitForNumberOfTiles(int32 InSize, int32& OutRatioX, int32& OutRatioY);
 
-EPixelFormatChannelFlags GetPixelFormatChannelFlagForButton(ETextureChannelButton InButton)
+static EPixelFormatChannelFlags GetPixelFormatChannelFlagForButton(ETextureChannelButton InButton)
 {
 	switch (InButton)
 	{
@@ -103,8 +103,35 @@ EPixelFormatChannelFlags GetPixelFormatChannelFlagForButton(ETextureChannelButto
 	return EPixelFormatChannelFlags::None;
 }
 
+// returns true if you should call PostEditChange/UpdateResource/etc to re-compress the texture
+//	 after changing properties ; if false the PlatformData should be left as-is
+static bool CanRecompressTexture(UTexture * Texture)
+{
+	if ( ! Texture->Source.IsValid() )
+	{
+		return false;
+	}
+
+	// legacy code was doing this to identify rendertargets :
+	// not sure this is doing anything useful, the source.isvalid check above does all the work
+	ETextureClass TextureClass = Texture->GetTextureClass();
+	if (TextureClass == ETextureClass::RenderTarget ||
+		TextureClass == ETextureClass::Other2DNoSource ||
+		TextureClass == ETextureClass::TwoDDynamic)
+	{
+		return false;
+	}
+
+	return true;
+}
+
 void FTextureEditorToolkit::PostTextureRecode()
 {
+	if ( ! CanRecompressTexture(Texture) )
+	{
+		return;
+	}
+
 	// Each time we change a custom encode setting we want to re-encode the texture
 	// as though we changed a compression setting on the actual texture, so we just
 	// post a CompressionSettings property changed event to handle all of that for
@@ -146,11 +173,28 @@ FTextureEditorToolkit::~FTextureEditorToolkit( )
 
 	GEditor->UnregisterForUndo(this);
 
-	if (CustomEncoding->bUseCustomEncode)
+	Texture->PreEditChange(nullptr);	
 	{
-		// reencode the texture with normal settings.
+		// we are leaving the texture editor
+		// restore any temporary encoding settings we may have changed
+	
+		Texture->DeferCompression = false;
+
+		// we could leave OverrideRunningPlatformName set (like CompressFinal)
+		//	it will stay set for this Editor session, it is not serialized
+		// for now let's clear it for consistency's sake
+		Texture->OverrideRunningPlatformName = NAME_None;
+
 		CustomEncoding->bUseCustomEncode = false;
-		PostTextureRecode();
+		Texture->TextureEditorCustomEncoding = nullptr;
+	
+		// Texture->CompressFinal intentionally not changed
+		//	it will stay set for this Editor session, it is not serialized
+	}
+
+	if ( CanRecompressTexture(Texture) )
+	{
+		Texture->PostEditChange();
 	}
 }
 
@@ -217,6 +261,15 @@ void FTextureEditorToolkit::InitTextureEditor( const EToolkitMode::Type Mode, co
 	GEditor->RegisterForUndo(this);
 
 	CustomEncoding = MakeShared<FTextureEditorCustomEncode>(FTextureEditorCustomEncode());
+	
+	check( CustomEncoding->bUseCustomEncode == false );
+
+	// OpenAssetEditor should ensure that we never have two toolkits open for the same asset!
+	check(Texture->TextureEditorCustomEncoding == nullptr);
+
+	// We save this as a separate object so that the engine can reference the type
+	// without needing the texture editor module.
+	Texture->TextureEditorCustomEncoding = CustomEncoding;
 
 	// initialize view options
 	bIsRedChannel = true;
@@ -334,16 +387,7 @@ void FTextureEditorToolkit::InitTextureEditor( const EToolkitMode::Type Mode, co
 	// because it's transient, and any async build was completed above.
 	Texture->CompressFinal = ( Texture->Source.IsValid() && ( Texture->Source.GetTotalTopMipPixelCount() <= (int64)4096*4096  ) );
 
-	// We don't want to post recodes for render targets because that clears them to black and
-	// we don't care about CompressFinal for them anyway as they aren't encoded. While we are
-	// here, don't bother with other dynamic textures as well.
-	ETextureClass TextureClass = Texture->GetTextureClass();
-	if (TextureClass != ETextureClass::RenderTarget &&
-		TextureClass != ETextureClass::Other2DNoSource &&
-		TextureClass != ETextureClass::TwoDDynamic)
-	{
-		PostTextureRecode();
-	}
+	PostTextureRecode();
 
 	// @todo toolkit world centric editing
 	/*if(IsWorldCentricAssetEditor())
@@ -2003,7 +2047,8 @@ void FTextureEditorToolkit::FillToolbar(FToolBarBuilder& ToolbarBuilder)
 		ToolbarBuilder.EndSection();
 
 		if (!Texture->VirtualTextureStreaming &&
-			Texture->Availability == ETextureAvailability::GPU)
+			Texture->Availability == ETextureAvailability::GPU &&
+			CanRecompressTexture(Texture))
 		{
 			ToolbarBuilder.BeginSection("PlatformPreview");
 			{
@@ -2360,11 +2405,18 @@ bool FTextureEditorToolkit::HandleCubemapViewModeActionIsChecked(ETextureEditorC
 
 void FTextureEditorToolkit::HandleCompressNowActionExecute( )
 {
+	if ( ! CanRecompressTexture(Texture) )
+	{
+		return;
+	}
+
 	GWarn->BeginSlowTask(NSLOCTEXT("TextureEditor", "CompressNow", "Compressing 1 Textures that have Defer Compression set"), true);
 
 	if (Texture->DeferCompression)
 	{
+		Texture->PreEditChange(nullptr);
 		// turn off deferred compression and compress the texture
+		// semi-code dupe of UTexture::PreSave
 		Texture->DeferCompression = false;
 		Texture->Source.Compress();
 		Texture->PostEditChange();
@@ -2805,7 +2857,7 @@ void FTextureEditorToolkit::EditorOodleSettingsEffortChanged(int32 NewValue, ESe
 
 	CustomEncoding->OodleEncodeEffort = IntCastChecked<uint8>(NewValue);
 
-	if (CustomEncoding->bUseCustomEncode || bChanged)
+	if (CustomEncoding->bUseCustomEncode && bChanged)
 	{
 		PostTextureRecode();
 	}
@@ -3300,9 +3352,13 @@ TSharedRef<SWidget> FTextureEditorToolkit::MakePlatformSelectorWidget()
 							int32 Index = AvailablePlatforms.Find(NewPlatformName);
 							if (Index != INDEX_NONE)
 							{
-								Texture->OverrideRunningPlatformName = AvailablePlatformNames[Index];
 								ViewingPlatform = AvailablePlatformNames[Index];
-								Texture->UpdateResource();
+								if( ViewingPlatform != Texture->OverrideRunningPlatformName )
+								{
+									Texture->PreEditChange(nullptr);
+									Texture->OverrideRunningPlatformName = ViewingPlatform;
+									Texture->PostEditChange();
+								}
 							}
 						})
 				]
