Account for LOD bias when determining volume depth slice count in texture editor

diff --git a/Engine/Source/Editor/TextureEditor/Private/Models/TextureEditorViewportClient.cpp b/Engine/Source/Editor/TextureEditor/Private/Models/TextureEditorViewportClient.cpp
--- a/Engine/Source/Editor/TextureEditor/Private/Models/TextureEditorViewportClient.cpp
+++ b/Engine/Source/Editor/TextureEditor/Private/Models/TextureEditorViewportClient.cpp
@@ -248,7 +248,7 @@ void FTextureEditorViewportClient::Draw(FViewport* Viewport, FCanvas* Canvas)
 		{
 			BatchedElementParameters = new FBatchedElementVolumeTexturePreviewParameters(
 				TextureEditorPinned->GetVolumeViewMode() == TextureEditorVolumeViewMode_DepthSlices,
-				FMath::Max<int32>(VolumeTexture->GetSizeZ(), 1), 
+				FMath::Max<int32>(VolumeTexture->GetSizeZ() >> VolumeTexture->GetCachedLODBias(), 1), 
 				MipLevel, 
 				(float)TextureEditorPinned->GetVolumeOpacity(),
 				true, 
