TextureEditorToolkit : fix possible null deref for partially valid target platforms

diff --git a/Engine/Source/Editor/TextureEditor/Private/TextureEditorToolkit.cpp b/Engine/Source/Editor/TextureEditor/Private/TextureEditorToolkit.cpp
--- a/Engine/Source/Editor/TextureEditor/Private/TextureEditorToolkit.cpp
+++ b/Engine/Source/Editor/TextureEditor/Private/TextureEditorToolkit.cpp
@@ -3287,12 +3287,18 @@ TSharedRef<SWidget> FTextureEditorToolkit::MakePlatformSelectorWidget()
 			}
 
 			const PlatformInfo::FTargetPlatformInfo* VanillaInfo = PlatformInfo::FindVanillaPlatformInfo(Pair.Key);
-			const TArray<const PlatformInfo::FTargetPlatformInfo*> ValidFlavors = VanillaInfo->Flavors.FilterByPredicate([](const PlatformInfo::FTargetPlatformInfo* Target)
+			// VanillaInfo can come back null for platforms that are partially valid
+			if ( VanillaInfo == nullptr )
 			{
-				// Editor isn't a valid platform type that users can target
-				// The Build Target will choose client or server, so no need to show them as well
-				return Target->PlatformType != EBuildTargetType::Editor && Target->PlatformType != EBuildTargetType::Client && Target->PlatformType != EBuildTargetType::Server;
-			});
+				continue;
+			}
+
+			const TArray<const PlatformInfo::FTargetPlatformInfo*> ValidFlavors = VanillaInfo->Flavors.FilterByPredicate([](const PlatformInfo::FTargetPlatformInfo* Target)
+				{
+					// Editor isn't a valid platform type that users can target
+					// The Build Target will choose client or server, so no need to show them as well
+					return Target->PlatformType != EBuildTargetType::Editor && Target->PlatformType != EBuildTargetType::Client && Target->PlatformType != EBuildTargetType::Server;
+				});
 
 			if (ValidFlavors.Num())
 			{
