Fix for IBuild error messages not completely getting shown in the texture editor viewer.

diff --git a/Engine/Source/Editor/TextureEditor/Private/Models/TextureEditorViewportClient.cpp b/Engine/Source/Editor/TextureEditor/Private/Models/TextureEditorViewportClient.cpp
--- a/Engine/Source/Editor/TextureEditor/Private/Models/TextureEditorViewportClient.cpp
+++ b/Engine/Source/Editor/TextureEditor/Private/Models/TextureEditorViewportClient.cpp
@@ -66,6 +66,40 @@ struct FTextureErrorLogger : public FOutputDevice
 			return;
 		}
 
+
+		// Here we aren't capturing yet.
+		// See if the string relates to us.
+		if (FCString::Stristr(V, *TextureToMonitor->GetName()))
+		{
+			// It turns out we can't rely on a startup message with the new build flow (IBuild), so
+			// we just start capturing if we ever see a warning or error with our name in it.
+			// 
+			// If it's any of the launching messages then we clear our backlog.
+			if (Verbosity == ELogVerbosity::Error || 
+				Verbosity == ELogVerbosity::Warning)
+			{
+				bCurrentlyCapturing = true;
+			}
+
+			// If it's "Building textures" then we started a new build and need to empty our list.
+			// For shared linear we also might only get a "Tiling" message... but we don't want to clear
+			// on that one since we'll likely get it after the Building textures message
+			if (FCString::Stristr(V, TEXT("Building textures")))
+			{
+				RelevantLogLines.Empty();
+				bCurrentlyCapturing = true;
+			}
+			static FName TextureBuildFunctionCategory = "LogTextureBuildFunction";
+			if (Category == TextureBuildFunctionCategory)
+			{
+				if (FCString::Stristr(V, TEXT("Compressing")))
+				{
+					RelevantLogLines.Empty();
+					bCurrentlyCapturing = true;
+				}
+			}
+		}
+
 		//
 		// Error messages don't reliably put the texture name in the messages, and we don't necessarily know
 		// that the error will come from a texture category if it's something like bulk data. However, we do
@@ -110,23 +144,6 @@ struct FTextureErrorLogger : public FOutputDevice
 			return;
 		}
 
-		// Here we aren't capturing yet.
-		// See if the string relates to us.
-		if (FCString::Stristr(V, *TextureToMonitor->GetName()))
-		{
-			// If it's "Building textures" then we started a new build and need to empty our list.
-			// For shared linear we also might only get a "Tiling" message... but we don't want to clear
-			// on that one since we'll likely get it after the Building textures message
-			if (FCString::Stristr(V, TEXT("Building textures")))
-			{
-				RelevantLogLines.Empty();
-				bCurrentlyCapturing = true;
-			}
-			else if (FCString::Stristr(V, TEXT("Tiling")))
-			{
-				bCurrentlyCapturing = true;
-			}
-		}
 	}
 };
 
