TextureEditorToolkit : fix bug that caused textures to be marked as modified any time you open the texture editor viewer

diff --git a/Engine/Source/Editor/TextureEditor/Private/TextureEditorToolkit.cpp b/Engine/Source/Editor/TextureEditor/Private/TextureEditorToolkit.cpp
--- a/Engine/Source/Editor/TextureEditor/Private/TextureEditorToolkit.cpp
+++ b/Engine/Source/Editor/TextureEditor/Private/TextureEditorToolkit.cpp
@@ -136,6 +136,7 @@ void FTextureEditorToolkit::PostTextureRecode()
 	// as though we changed a compression setting on the actual texture, so we just
 	// post a CompressionSettings property changed event to handle all of that for
 	// us.
+	// @@ this is a bit odd, why is it done this way rather than just calling PostEditChange() ?
 	FProperty* Property = FindFProperty<FProperty>(UTexture::StaticClass(), "CompressionSettings");
 	FPropertyChangedEvent PropertyChangedEvent(Property);
 	Texture->PostEditChangeProperty(PropertyChangedEvent);
@@ -173,7 +174,7 @@ FTextureEditorToolkit::~FTextureEditorToolkit( )
 
 	GEditor->UnregisterForUndo(this);
 
-	Texture->PreEditChange(nullptr);	
+	Texture->BlockOnAnyAsyncBuild(); // PreEditChange , but don't mark as modified ; same as Modify(false)
 	{
 		// we are leaving the texture editor
 		// restore any temporary encoding settings we may have changed
@@ -2412,13 +2413,13 @@ void FTextureEditorToolkit::HandleCompressNowActionExecute( )
 
 	GWarn->BeginSlowTask(NSLOCTEXT("TextureEditor", "CompressNow", "Compressing 1 Textures that have Defer Compression set"), true);
 
+	// turn off deferred compression and compress the texture
+
 	if (Texture->DeferCompression)
 	{
-		Texture->PreEditChange(nullptr);
-		// turn off deferred compression and compress the texture
-		// semi-code dupe of UTexture::PreSave
+		Texture->BlockOnAnyAsyncBuild(); // PreEditChange , but don't mark as modified ; same as Modify(false)
 		Texture->DeferCompression = false;
-		Texture->Source.Compress();
+		//Texture->Source.Compress(); // <- done in UTexture::PreSave
 		Texture->PostEditChange();
 
 		PopulateQuickInfo();
@@ -3361,7 +3362,7 @@ TSharedRef<SWidget> FTextureEditorToolkit::MakePlatformSelectorWidget()
 								ViewingPlatform = AvailablePlatformNames[Index];
 								if( ViewingPlatform != Texture->OverrideRunningPlatformName )
 								{
-									Texture->PreEditChange(nullptr);
+									Texture->BlockOnAnyAsyncBuild(); // PreEditChange , but don't mark as modified ; same as Modify(false)
 									Texture->OverrideRunningPlatformName = ViewingPlatform;
 									Texture->PostEditChange();
 								}
