TextureEditorToolkit : don't block on texture build in GUI exit if it's not necessary

diff --git a/Engine/Source/Editor/TextureEditor/Private/TextureEditorToolkit.cpp b/Engine/Source/Editor/TextureEditor/Private/TextureEditorToolkit.cpp
--- a/Engine/Source/Editor/TextureEditor/Private/TextureEditorToolkit.cpp
+++ b/Engine/Source/Editor/TextureEditor/Private/TextureEditorToolkit.cpp
@@ -173,11 +173,15 @@ FTextureEditorToolkit::~FTextureEditorToolkit( )
 	GEditor->GetEditorSubsystem<UImportSubsystem>()->OnAssetPostImport.RemoveAll(this);
 
 	GEditor->UnregisterForUndo(this);
+	
+	// we are leaving the texture editor
+	// restore any temporary encoding settings we may have changed
 
-	Texture->BlockOnAnyAsyncBuild(); // PreEditChange , but don't mark as modified ; same as Modify(false)
+	if ( Texture->DeferCompression || 
+		Texture->OverrideRunningPlatformName != NAME_None ||
+		CustomEncoding->bUseCustomEncode )
 	{
-		// we are leaving the texture editor
-		// restore any temporary encoding settings we may have changed
+		Texture->BlockOnAnyAsyncBuild(); // PreEditChange , but don't mark as modified ; same as Modify(false)
 	
 		Texture->DeferCompression = false;
 
@@ -191,11 +195,11 @@ FTextureEditorToolkit::~FTextureEditorToolkit( )
 	
 		// Texture->CompressFinal intentionally not changed
 		//	it will stay set for this Editor session, it is not serialized
-	}
-
-	if ( CanRecompressTexture(Texture) )
-	{
-		Texture->PostEditChange();
+		
+		if ( CanRecompressTexture(Texture) )
+		{
+			Texture->PostEditChange();
+		}
 	}
 }
 
@@ -254,7 +258,7 @@ void FTextureEditorToolkit::InitTextureEditor( const EToolkitMode::Type Mode, co
 	Texture = CastChecked<UTexture>(ObjectToEdit);
 
 	// The texture being edited might still be compiling, wait till it finishes then.
-	// FinishCompilation is nice enough to provide a progress for us while we're waiting.
+	// @@ is this necessary? can we remove it?
 	Texture->BlockOnAnyAsyncBuild();
 
 	// Support undo/redo
