Fix 1-frame glitch in texture editor when changing mip level

diff --git a/Engine/Source/Editor/TextureEditor/Private/TextureEditorToolkit.cpp b/Engine/Source/Editor/TextureEditor/Private/TextureEditorToolkit.cpp
--- a/Engine/Source/Editor/TextureEditor/Private/TextureEditorToolkit.cpp
+++ b/Engine/Source/Editor/TextureEditor/Private/TextureEditorToolkit.cpp
@@ -2475,6 +2475,7 @@ bool FTextureEditorToolkit::HandleMipLevelCheckBoxIsEnabled( ) const
 void FTextureEditorToolkit::HandleMipLevelChanged(int32 NewMipLevel)
 {
 	SpecifiedMipLevel = FMath::Clamp<int32>(NewMipLevel, MIPLEVEL_MIN, GetMaxMipLevel().Get(MIPLEVEL_MAX));
+	PopulateQuickInfo(); // so PreviewEffectiveTexture{Width,Height} get updated immediately
 
 	MipLevelTextBlock->SetText(FText::Format(LOCTEXT("MipLevel", "Mip Level {0}"), SpecifiedMipLevel));
 }
