Fix crash due to race condition accessing stored log lines in texture editor viewer.

diff --git a/Engine/Source/Editor/TextureEditor/Private/Models/TextureEditorViewportClient.cpp b/Engine/Source/Editor/TextureEditor/Private/Models/TextureEditorViewportClient.cpp
--- a/Engine/Source/Editor/TextureEditor/Private/Models/TextureEditorViewportClient.cpp
+++ b/Engine/Source/Editor/TextureEditor/Private/Models/TextureEditorViewportClient.cpp
@@ -39,6 +39,7 @@ static TAutoConsoleVariable<int32> CVarEnableVTFeedback(
 struct FTextureErrorLogger : public FOutputDevice
 {
 	UTexture* TextureToMonitor = nullptr;
+	FCriticalSection LogLinesLock; // Locks RelevantLogLings since Serialize can be called while we are drawing the lines.
 	TArray<TPair<bool, FString>> RelevantLogLines;
 
 	double CaptureUntilTime = 0;
@@ -86,6 +87,7 @@ struct FTextureErrorLogger : public FOutputDevice
 			// on that one since we'll likely get it after the Building textures message
 			if (FCString::Stristr(V, TEXT("Building textures")))
 			{
+				FScopeLock _(&LogLinesLock);
 				RelevantLogLines.Empty();
 				bCurrentlyCapturing = true;
 			}
@@ -94,6 +96,7 @@ struct FTextureErrorLogger : public FOutputDevice
 			{
 				if (FCString::Stristr(V, TEXT("Compressing")))
 				{
+					FScopeLock _(&LogLinesLock);
 					RelevantLogLines.Empty();
 					bCurrentlyCapturing = true;
 				}
@@ -127,19 +130,24 @@ struct FTextureErrorLogger : public FOutputDevice
 			}
 
 			// Add any errors or warnings to the list
-			if (Verbosity == ELogVerbosity::Error)
-			{
-				RelevantLogLines.Add(TPair<bool, FString>(true, FString(V)));
-			}
-			else if (Verbosity == ELogVerbosity::Warning)
-			{
-				RelevantLogLines.Add(TPair<bool, FString>(false, FString(V)));
-			}
-			if (RelevantLogLines.Num() == 10)
+			if (Verbosity == ELogVerbosity::Error ||
+				Verbosity == ELogVerbosity::Warning)
 			{
-				RelevantLogLines.Add(TPair<bool, FString>(false, TEXT("Too much to show: check Output Log")));
-				bCurrentlyCapturing = false;
-				CaptureUntilTime = 0;
+				FScopeLock _(&LogLinesLock);
+				if (Verbosity == ELogVerbosity::Error)
+				{
+					RelevantLogLines.Add(TPair<bool, FString>(true, FString(V)));
+				}
+				else if (Verbosity == ELogVerbosity::Warning)
+				{
+					RelevantLogLines.Add(TPair<bool, FString>(false, FString(V)));
+				}
+				if (RelevantLogLines.Num() == 10)
+				{
+					RelevantLogLines.Add(TPair<bool, FString>(false, TEXT("Too much to show: check Output Log")));
+					bCurrentlyCapturing = false;
+					CaptureUntilTime = 0;
+				}
 			}
 			return;
 		}
@@ -508,10 +516,13 @@ void FTextureEditorViewportClient::Draw(FViewport* Viewport, FCanvas* Canvas)
 	} // end if not deferring
 
 	// Print any warnings/errors that we saw in the output log.
-	for (TPair<bool, FString>& ReportedLine : TextureConsoleCapture->RelevantLogLines)
 	{
-		Canvas->DrawShadowedText(ReportingLineX, ReportingLineY, FText::FromString(ReportedLine.Value), GEngine->GetLargeFont(), ReportedLine.Key ? FLinearColor::Red : FLinearColor::Yellow);
-		ReportingLineY += ReportingLineHeight;
+		FScopeLock _(&TextureConsoleCapture->LogLinesLock);
+		for (TPair<bool, FString>& ReportedLine : TextureConsoleCapture->RelevantLogLines)
+		{
+			Canvas->DrawShadowedText(ReportingLineX, ReportingLineY, FText::FromString(ReportedLine.Value), GEngine->GetLargeFont(), ReportedLine.Key ? FLinearColor::Red : FLinearColor::Yellow);
+			ReportingLineY += ReportingLineHeight;
+		}
 	}
 }
 
