Texture GetMaximumDimension nop/comments only

diff --git a/Engine/Source/Editor/TextureEditor/Private/Customizations/TextureDetailsCustomization.cpp b/Engine/Source/Editor/TextureEditor/Private/Customizations/TextureDetailsCustomization.cpp
--- a/Engine/Source/Editor/TextureEditor/Private/Customizations/TextureDetailsCustomization.cpp
+++ b/Engine/Source/Editor/TextureEditor/Private/Customizations/TextureDetailsCustomization.cpp
@@ -92,14 +92,20 @@ void FTextureDetails::CustomizeDetails(IDetailLayoutBuilder& DetailBuilder)
 		FDetailWidgetRow Row;
 		MaxTextureSizePropertyRow.GetDefaultWidgets(NameWidget, ValueWidget, Row);
 
-		int32 MaxTextureSize = 2048;
+		int32 MaxTextureSize = UTexture::GetMaximumDimensionOfNonVT();
 
 		if (UTexture* Texture = Cast<UTexture>(TexturesBeingCustomized[0].Get()))
 		{
 			// GetMaximumDimension is for current RHI and texture type
-			MaxTextureSize = Texture->GetMaximumDimension();
+			MaxTextureSize = FMath::Min<int32>( Texture->GetMaximumDimension(), MaxTextureSize );
 		}
 
+		// @@ this slider is very hard to work with
+		//	it's almost impossible to set low values
+		// instead of being on the linear MaxTextureSize value, it should be on the log2
+		//	and scaled by *10 or something
+		// so the drag experience is slower and log-scaled
+
 		const bool bShowChildren = true;
 		MaxTextureSizePropertyRow.CustomWidget(bShowChildren)
 			.NameContent()
@@ -181,6 +187,7 @@ void FTextureDetails::OnMaxTextureSizeChanged(int32 NewValue)
 		}
 
 		// We don't create a transaction for each property change when using the slider.  Only once when the slider first is moved
+		// Interactive flag makes it so the texture is not rebuilt in PostEditChange
 		EPropertyValueSetFlags::Type Flags = (EPropertyValueSetFlags::InteractiveChange | EPropertyValueSetFlags::NotTransactable);
 		MaxTextureSizePropertyHandle->SetValue(NewValue, Flags);
 	}
@@ -188,6 +195,7 @@ void FTextureDetails::OnMaxTextureSizeChanged(int32 NewValue)
 
 void FTextureDetails::OnMaxTextureSizeCommitted(int32 NewValue, ETextCommit::Type CommitInfo)
 {
+	// this causes the texture to build with the new value (if necessary)
 	MaxTextureSizePropertyHandle->SetValue(NewValue);
 }
 
diff --git a/Engine/Source/Runtime/Engine/Private/Texture.cpp b/Engine/Source/Runtime/Engine/Private/Texture.cpp
--- a/Engine/Source/Runtime/Engine/Private/Texture.cpp
+++ b/Engine/Source/Runtime/Engine/Private/Texture.cpp
@@ -3581,6 +3581,11 @@ void FTextureSource::SetId(const FGuid& InId, bool bInGuidIsHash)
 uint32 UTexture::GetMaximumDimension() const
 {
 	// the various virtual implementations of this wind up returning GRHIGlobals.MaxCubeTextureDimensions etc.
+	//
+	// BEWARE : this can be higher than GetMaximumDimensionOfNonVT() , but you don't actually want that!
+	//	probably this should be doing Min(GetMaximumDimensionOfNonVT,*) here so that values over GetMaximumDimensionOfNonVT
+	//	are never returned out of here
+	// because it does not, you should always do that Min on the usage side
 
 	// just assume anyone who doesn't implement this virtual is 2d
 	return GetMax2DTextureDimension();
