TextureEditor:
+ Add button to detect source alpha if unknown
+ Clarify "Has Alpha Channel" and move to next to Format.
+ Platform Preview is always visible and enables/disables based on properties.

diff --git a/Engine/Source/Editor/TextureEditor/Private/TextureEditorToolkit.cpp b/Engine/Source/Editor/TextureEditor/Private/TextureEditorToolkit.cpp
--- a/Engine/Source/Editor/TextureEditor/Private/TextureEditorToolkit.cpp
+++ b/Engine/Source/Editor/TextureEditor/Private/TextureEditorToolkit.cpp
@@ -45,6 +45,7 @@
 #include "Curves/CurveLinearColorAtlas.h"
 #include "TextureEditorSettings.h"
 #include "Widgets/Input/SSlider.h"
+#include "Widgets/Input/SButton.h"
 #include "Widgets/Input/STextComboBox.h"
 #include "Widgets/Layout/SSpacer.h"
 #include "Menus/TextureEditorViewOptionsMenu.h"
@@ -594,7 +595,7 @@ void FTextureEditorToolkit::PopulateQuickInfo( )
 		LODBiasText->SetText(NSLOCTEXT("TextureEditor", "QuickInfo_LODBias_NA", "Combined LOD Bias: Computing..."));
 		FormatText->SetText(NSLOCTEXT("TextureEditor", "QuickInfo_Format_NA", "Format: Computing..."));
 		NumMipsText->SetText(NSLOCTEXT("TextureEditor", "QuickInfo_NumMips_NA", "Number of Mips: Computing..."));
-		HasAlphaChannelText->SetText(NSLOCTEXT("TextureEditor", "QuickInfo_HasAlphaChannel_NA", "Has Alpha Channel: Computing..."));
+		HasAlphaChannelText->SetText(NSLOCTEXT("TextureEditor", "QuickInfo_HasAlphaChannel_NA", "Format Supports Alpha: Computing..."));
 		EncodeSpeedText->SetText(NSLOCTEXT("TextureEditor", "QuickInfo_EncodeSpeed_Computing", "Encode Speed: Computing..."));
 		SceneCaptureSizeText->SetText(FText());
 		SceneCaptureNameText->SetText(FText());
@@ -810,11 +811,13 @@ void FTextureEditorToolkit::PopulateQuickInfo( )
 			bool bSourceAlphaDetected = FImageCore::DetectAlphaChannel(View);
 			SourceMipsAlphaDetectedText->SetText(FText::Format(NSLOCTEXT("TextureEditor", "QuickInfo_SourceAlphaDetected", "Source Alpha Detected: {0}"),
 				bSourceAlphaDetected ? NSLOCTEXT("TextureEditor", "True", "True") : NSLOCTEXT("TextureEditor", "False", "False")));
+			DetectSourceAlphaButton->SetVisibility(EVisibility::Hidden);
 		}
 		else
 		{
 			SourceMipsAlphaDetectedText->SetText(FText::Format(NSLOCTEXT("TextureEditor", "QuickInfo_SourceAlphaDetected", "Source Alpha Detected: {0}"),
 				NSLOCTEXT("TextureEditor", "Unknown", "Unknown")));
+			DetectSourceAlphaButton->SetVisibility(Texture->Source.IsValid() ? EVisibility::Visible : EVisibility::Hidden);
 		}
 	} // end if valid platform data
 
@@ -973,7 +976,7 @@ void FTextureEditorToolkit::PopulateQuickInfo( )
 	// This "Has Alpha Channel" is whether the GPU format can represent alpha in the format (eg. is it DXT1 vs DXT5)
 	//	it does not tell you if the texture actually has non-opaque alpha
 	EPixelFormatChannelFlags ValidTextureChannels = GetPixelFormatValidChannels(TextureFormat);
-	HasAlphaChannelText->SetText(FText::Format(NSLOCTEXT("TextureEditor", "QuickInfo_HasAlphaChannel", "Has Alpha Channel: {0}"),
+	HasAlphaChannelText->SetText(FText::Format(NSLOCTEXT("TextureEditor", "QuickInfo_HasAlphaChannel", "Format Supports Alpha: {0}"),
 		EnumHasAnyFlags(ValidTextureChannels, EPixelFormatChannelFlags::A) ? NSLOCTEXT("TextureEditor", "True", "True") : NSLOCTEXT("TextureEditor", "False", "False")));
 	HasAlphaChannelText->SetVisibility((ValidTextureChannels != EPixelFormatChannelFlags::None) ? EVisibility::Visible : EVisibility::Collapsed);
 
@@ -1856,15 +1859,24 @@ void FTextureEditorToolkit::CreateInternalWidgets()
 			.VAlign(VAlign_Center)
 			.Padding(4.0f)
 			[
-				SAssignNew(HasAlphaChannelText, STextBlock)
-			]
+				SNew(SHorizontalBox)
 
-			+ SVerticalBox::Slot()
-			.AutoHeight()
-			.VAlign(VAlign_Center)
-			.Padding(4.0f)
-			[
-				SAssignNew(SourceMipsAlphaDetectedText, STextBlock)
+				+ SHorizontalBox::Slot()
+				.VAlign(VAlign_Center)
+				.AutoWidth()
+				[
+					SAssignNew(SourceMipsAlphaDetectedText, STextBlock)
+				]
+
+				+ SHorizontalBox::Slot()
+				.VAlign(VAlign_Center)
+				.Padding(4.0f, 0.0f, 0.0f, 0.0f)
+				.AutoWidth()
+				[
+					SAssignNew(DetectSourceAlphaButton, SButton)
+					.Text(LOCTEXT("DetectSourceAlpha_Button", "Detect"))
+					.OnClicked(this, &FTextureEditorToolkit::DetectSourceAlphaButton_Clicked)
+				]
 			]
 
 			+ SVerticalBox::Slot()
@@ -1897,6 +1909,16 @@ void FTextureEditorToolkit::CreateInternalWidgets()
 				SAssignNew(FormatText, STextBlock)
 			]
 
+
+			+ SVerticalBox::Slot()
+			.AutoHeight()
+			.VAlign(VAlign_Center)
+			.Padding(4.0f)
+			[
+				SAssignNew(HasAlphaChannelText, STextBlock)
+			]
+
+
 			+ SVerticalBox::Slot()
 			.AutoHeight()
 			.VAlign(VAlign_Center)
@@ -2051,9 +2073,6 @@ void FTextureEditorToolkit::FillToolbar(FToolBarBuilder& ToolbarBuilder)
 		}
 		ToolbarBuilder.EndSection();
 
-		if (!Texture->VirtualTextureStreaming &&
-			Texture->Availability == ETextureAvailability::GPU &&
-			CanRecompressTexture(Texture))
 		{
 			ToolbarBuilder.BeginSection("PlatformPreview");
 			{
@@ -3347,7 +3366,7 @@ TSharedRef<SWidget> FTextureEditorToolkit::MakePlatformSelectorWidget()
 				[
 					SNew(STextBlock)
 						.Text(LOCTEXT("PlatformPreviewLabel", "Preview Platform:"))
-						.ToolTipText(LOCTEXT("PlatformPreviewTT", "If a platform is chosen, the texture will be encoded as though cooked for that platform, then if necessary decoded so that it can be viewed on this platform."))
+						.ToolTipText(LOCTEXT("PlatformPreviewTT", "If a platform is chosen, the texture will be encoded as though cooked for that platform, then if necessary decoded so that it can be viewed on this platform. Requires: Texture source, non-virtual texture, and GPU availability."))
 				]
 				+ SHorizontalBox::Slot()
 				.VAlign(VAlign_Center)
@@ -3357,6 +3376,7 @@ TSharedRef<SWidget> FTextureEditorToolkit::MakePlatformSelectorWidget()
 					SNew(STextComboBox)
 						.OptionsSource(&AvailablePlatforms)
 						.InitiallySelectedItem(InitialSelection)
+						.IsEnabled(this, &FTextureEditorToolkit::CanPlatformPreview)
 						.OnSelectionChanged_Lambda(
 							[this]
 							(TSharedPtr<FString> NewPlatformName, ESelectInfo::Type SelectInfo)
@@ -3546,4 +3566,27 @@ void FTextureEditorToolkit::PackagingSettingsChanged(TSharedPtr<FString> Selecti
 	}
 }
 
+FReply FTextureEditorToolkit::DetectSourceAlphaButton_Clicked()
+{
+	if (!Texture ||
+		!Texture->Source.IsValid() ||
+		!Texture->IsAsyncCacheComplete())
+	{
+		return FReply::Handled();
+	}
+
+	Texture->PreEditChange(nullptr);
+	Texture->Source.UpdateChannelLinearMinMax();
+	Texture->PostEditChange();
+	return FReply::Handled();
+}
+
+bool FTextureEditorToolkit::CanPlatformPreview() const
+{
+	return Texture && 
+		!Texture->VirtualTextureStreaming &&
+		Texture->Availability == ETextureAvailability::GPU &&
+		CanRecompressTexture(Texture);
+}
+
 #undef LOCTEXT_NAMESPACE
diff --git a/Engine/Source/Editor/TextureEditor/Private/TextureEditorToolkit.h b/Engine/Source/Editor/TextureEditor/Private/TextureEditorToolkit.h
--- a/Engine/Source/Editor/TextureEditor/Private/TextureEditorToolkit.h
+++ b/Engine/Source/Editor/TextureEditor/Private/TextureEditorToolkit.h
@@ -15,6 +15,7 @@
 
 class SMenuAnchor;
 class STextBlock;
+class SButton;
 class STextureEditorViewport;
 class UFactory;
 class UTexture;
@@ -415,6 +416,7 @@ private:
 	TSharedPtr<STextBlock> LODBiasText;
 	TSharedPtr<STextBlock> HasAlphaChannelText;
 	TSharedPtr<STextBlock> SourceMipsAlphaDetectedText;
+	TSharedPtr<SButton> DetectSourceAlphaButton;
 	TSharedPtr<STextBlock> NumMipsText;
 	TSharedPtr<STextBlock> MipLevelTextBlock;
 	TSharedPtr<STextBlock> EncodeSpeedText;
@@ -514,6 +516,9 @@ private:
 	void OnUseEditorOodleSettingsChanged(ECheckBoxState NewState);
 	ECheckBoxState UseEditorOodleSettingsChecked() const;
 
+	FReply DetectSourceAlphaButton_Clicked();
+	bool CanPlatformPreview() const;
+
 	// callback for enabling custom encoding sub-controls
 	bool EditorOodleSettingsEnabled() const;
 
