Texture Editor : don't make whole GUI flash to blank while texture is building
instead show "Compiling..." text in the preview pane, but the rest of the GUI remains functional

diff --git a/Engine/Source/Editor/TextureEditor/Private/Models/TextureEditorViewportClient.cpp b/Engine/Source/Editor/TextureEditor/Private/Models/TextureEditorViewportClient.cpp
--- a/Engine/Source/Editor/TextureEditor/Private/Models/TextureEditorViewportClient.cpp
+++ b/Engine/Source/Editor/TextureEditor/Private/Models/TextureEditorViewportClient.cpp
@@ -202,8 +202,31 @@ void FTextureEditorViewportClient::Draw(FViewport* Viewport, FCanvas* Canvas)
 	
 	UpdateScrollBars();
 
-
 	Canvas->Clear( Settings.BackgroundColor );
+	
+	UFont* ReportingFont = GEngine->GetLargeFont();
+	const int32 ReportingLineHeight = FMath::CeilToInt(ReportingFont->GetMaxCharHeight()) + 2; // 2 for line spacing
+	const int32 ReportingLineX = 8;
+	int32 ReportingLineY = 8;
+
+	// make sure error messages are shown last, no matter if we return early or not
+	ON_SCOPE_EXIT {
+		// Print any warnings/errors that we saw in the output log.
+		FScopeLock _(&TextureConsoleCapture->LogLinesLock);
+		for (TPair<bool, FString>& ReportedLine : TextureConsoleCapture->RelevantLogLines)
+		{
+			Canvas->DrawShadowedText(ReportingLineX, ReportingLineY, FText::FromString(ReportedLine.Value), ReportingFont, ReportedLine.Key ? FLinearColor::Red : FLinearColor::Yellow);
+			ReportingLineY += ReportingLineHeight;
+		}
+	};
+
+	if ( Texture->IsCompiling() )
+	{
+		const FText Message = NSLOCTEXT("TextureEditor", "Compiling", "Compiling...");
+		Canvas->DrawShadowedText(ReportingLineX, ReportingLineY, Message, ReportingFont, FLinearColor::White);
+		ReportingLineY += ReportingLineHeight;
+		return;
+	}
 
 	UTexture2D* Texture2D = Cast<UTexture2D>(Texture);
 	UTextureCube* TextureCube = Cast<UTextureCube>(Texture);
@@ -389,12 +412,6 @@ void FTextureEditorViewportClient::Draw(FViewport* Viewport, FCanvas* Canvas)
 		}
 	}
 
-	UFont* ReportingFont = GEngine->GetLargeFont();
-	const int32 ReportingLineHeight = FMath::CeilToInt(ReportingFont->GetMaxCharHeight()) + 2; // 2 for line spacing
-	const int32 ReportingLineX = 8;
-	int32 ReportingLineY = 8;
-
-
 	// If we are requesting an explicit mip level of a VT asset, test to see if we can even display it properly and warn about it
 	if (bIsVirtualTexture && MipLevel >= 0.f)
 	{
@@ -515,15 +532,6 @@ void FTextureEditorViewportClient::Draw(FViewport* Viewport, FCanvas* Canvas)
 		} // end if valid result metadata
 	} // end if not deferring
 
-	// Print any warnings/errors that we saw in the output log.
-	{
-		FScopeLock _(&TextureConsoleCapture->LogLinesLock);
-		for (TPair<bool, FString>& ReportedLine : TextureConsoleCapture->RelevantLogLines)
-		{
-			Canvas->DrawShadowedText(ReportingLineX, ReportingLineY, FText::FromString(ReportedLine.Value), GEngine->GetLargeFont(), ReportedLine.Key ? FLinearColor::Red : FLinearColor::Yellow);
-			ReportingLineY += ReportingLineHeight;
-		}
-	}
 }
 
 
diff --git a/Engine/Source/Editor/TextureEditor/Private/TextureEditorToolkit.h b/Engine/Source/Editor/TextureEditor/Private/TextureEditorToolkit.h
--- a/Engine/Source/Editor/TextureEditor/Private/TextureEditorToolkit.h
+++ b/Engine/Source/Editor/TextureEditor/Private/TextureEditorToolkit.h
@@ -98,7 +98,10 @@ public:
 	{
 		return ExposureBias;
 	}
-
+	virtual EVisibility GetVisibilityWhileAssetCompiling() const override
+	{
+		return EVisibility::Visible;
+	}
 
 	// IToolkit interface
 	virtual FText GetBaseToolkitName( ) const override;
diff --git a/Engine/Source/Editor/UnrealEd/Private/Toolkits/AssetEditorToolkit.cpp b/Engine/Source/Editor/UnrealEd/Private/Toolkits/AssetEditorToolkit.cpp
--- a/Engine/Source/Editor/UnrealEd/Private/Toolkits/AssetEditorToolkit.cpp
+++ b/Engine/Source/Editor/UnrealEd/Private/Toolkits/AssetEditorToolkit.cpp
@@ -215,6 +215,8 @@ void FAssetEditorToolkit::InitAssetEditor( const EToolkitMode::Type Mode, const
 		this->TabManager = NewTabManager;
 
 		TArray<TWeakObjectPtr<UObject>> ObjectsToEditWeak;
+		EVisibility VisibilityWhileCompiling = GetVisibilityWhileAssetCompiling();
+
 		ObjectsToEditWeak.Reserve(ObjectsToEdit.Num());
 		for (UObject* Object : ObjectsToEdit)
 		{
@@ -223,7 +225,7 @@ void FAssetEditorToolkit::InitAssetEditor( const EToolkitMode::Type Mode, const
 		NewMajorTab->SetContent
 		( 
 			SAssignNew( NewStandaloneHost, SStandaloneAssetEditorToolkitHost, NewTabManager, AppIdentifier )
-			.Visibility_Lambda([ObjectsToEditWeak]()
+			.Visibility_Lambda([ObjectsToEditWeak,VisibilityWhileCompiling]()
 				{
 					for (const TWeakObjectPtr<UObject> Object : ObjectsToEditWeak)
 					{
@@ -231,7 +233,7 @@ void FAssetEditorToolkit::InitAssetEditor( const EToolkitMode::Type Mode, const
 						{
 							if (AsyncAsset->IsCompiling())
 							{
-								return EVisibility::Collapsed;
+								return VisibilityWhileCompiling;
 							}
 						}
 					}
diff --git a/Engine/Source/Editor/UnrealEd/Public/Toolkits/AssetEditorToolkit.h b/Engine/Source/Editor/UnrealEd/Public/Toolkits/AssetEditorToolkit.h
--- a/Engine/Source/Editor/UnrealEd/Public/Toolkits/AssetEditorToolkit.h
+++ b/Engine/Source/Editor/UnrealEd/Public/Toolkits/AssetEditorToolkit.h
@@ -371,6 +371,9 @@ protected:
 	/** Called to test if "Find in Content Browser" should be enabled for this asset */
 	virtual bool CanFindInContentBrowser() const { return true; }
 
+	/** Called to get the desired visibility when the asset is compiling */
+	UNREALED_API virtual EVisibility GetVisibilityWhileAssetCompiling() const { return EVisibility::Collapsed; }
+
 	/** Called to test if "Find in Content Browser" should be visible for this asset */
 	virtual bool IsFindInContentBrowserButtonVisible() const { return true; }
 
