Remove resource size when viewing platform preview textures for the moment.
#rnx

diff --git a/Engine/Source/Editor/TextureEditor/Private/TextureEditorToolkit.cpp b/Engine/Source/Editor/TextureEditor/Private/TextureEditorToolkit.cpp
--- a/Engine/Source/Editor/TextureEditor/Private/TextureEditorToolkit.cpp
+++ b/Engine/Source/Editor/TextureEditor/Private/TextureEditorToolkit.cpp
@@ -940,7 +940,16 @@ void FTextureEditorToolkit::PopulateQuickInfo( )
 		}
 	}
 
-	SizeText->SetText(FText::Format(NSLOCTEXT("TextureEditor", "QuickInfo_ResourceSize", "Resource Size: {0} KB"), FText::AsNumber(FMath::DivideAndRoundNearest(ResourceSize, (int64)1024), &FormatOptions)));
+	if (ViewingPlatform != NAME_None)
+	{
+		// Right now the resource size is the size of the texture we are viewing which could be decoded BGRA8 and not
+		// what you expect. This has been causing confusion so until I do UE-212930 we just make it clear we don't know.
+		SizeText->SetText(NSLOCTEXT("TextureEditor", "QuickInfo_ResourceSizeUnknown", "Resource Size: <unknown due to platform preview>"));
+	}
+	else
+	{
+		SizeText->SetText(FText::Format(NSLOCTEXT("TextureEditor", "QuickInfo_ResourceSize", "Resource Size: {0} KB"), FText::AsNumber(FMath::DivideAndRoundNearest(ResourceSize, (int64)1024), &FormatOptions)));
+	}
 
 	FText Method = Texture->IsCurrentlyVirtualTextured() ? NSLOCTEXT("TextureEditor", "QuickInfo_MethodVirtualStreamed", "Virtual Streamed")
 													: (!Texture->IsStreamable() ? NSLOCTEXT("TextureEditor", "QuickInfo_MethodNotStreamed", "Not Streamed") 
