Fixed log lines containing new lines displayed in the texture editor overlapping.

diff --git a/Engine/Source/Editor/TextureEditor/Private/Models/TextureEditorViewportClient.cpp b/Engine/Source/Editor/TextureEditor/Private/Models/TextureEditorViewportClient.cpp
--- a/Engine/Source/Editor/TextureEditor/Private/Models/TextureEditorViewportClient.cpp
+++ b/Engine/Source/Editor/TextureEditor/Private/Models/TextureEditorViewportClient.cpp
@@ -205,7 +205,8 @@ void FTextureEditorViewportClient::Draw(FViewport* Viewport, FCanvas* Canvas)
 	Canvas->Clear( Settings.BackgroundColor );
 	
 	UFont* ReportingFont = GEngine->GetLargeFont();
-	const int32 ReportingLineHeight = FMath::CeilToInt(ReportingFont->GetMaxCharHeight()) + 2; // 2 for line spacing
+	const int32 ReportingLineSpacing = 2;
+	const int32 ReportingLineHeight = FMath::CeilToInt(ReportingFont->GetMaxCharHeight()) + ReportingLineSpacing;
 	const int32 ReportingLineX = 8;
 	int32 ReportingLineY = 8;
 
@@ -215,8 +216,8 @@ void FTextureEditorViewportClient::Draw(FViewport* Viewport, FCanvas* Canvas)
 		FScopeLock _(&TextureConsoleCapture->LogLinesLock);
 		for (TPair<bool, FString>& ReportedLine : TextureConsoleCapture->RelevantLogLines)
 		{
-			Canvas->DrawShadowedText(ReportingLineX, ReportingLineY, FText::FromString(ReportedLine.Value), ReportingFont, ReportedLine.Key ? FLinearColor::Red : FLinearColor::Yellow);
-			ReportingLineY += ReportingLineHeight;
+			ReportingLineY += Canvas->DrawShadowedText(ReportingLineX, ReportingLineY, FText::FromString(ReportedLine.Value), ReportingFont, ReportedLine.Key ? FLinearColor::Red : FLinearColor::Yellow);
+			ReportingLineY += ReportingLineSpacing;
 		}
 	};
 
