Texture Editor: Fixed Source Alpha Detect button overlapping the rest of the UI.

diff --git a/Engine/Source/Editor/TextureEditor/Private/TextureEditorToolkit.cpp b/Engine/Source/Editor/TextureEditor/Private/TextureEditorToolkit.cpp
--- a/Engine/Source/Editor/TextureEditor/Private/TextureEditorToolkit.cpp
+++ b/Engine/Source/Editor/TextureEditor/Private/TextureEditorToolkit.cpp
@@ -1872,7 +1872,7 @@ void FTextureEditorToolkit::CreateInternalWidgets()
 
 				+ SHorizontalBox::Slot()
 				.VAlign(VAlign_Center)
-				.AutoWidth()
+				.FillWidth(1.f)
 				[
 					SAssignNew(SourceMipsAlphaDetectedText, STextBlock)
 				]
