Texture Editor: Collapse detect source alpha button so that it doesn't hide the text

diff --git a/Engine/Source/Editor/TextureEditor/Private/TextureEditorToolkit.cpp b/Engine/Source/Editor/TextureEditor/Private/TextureEditorToolkit.cpp
--- a/Engine/Source/Editor/TextureEditor/Private/TextureEditorToolkit.cpp
+++ b/Engine/Source/Editor/TextureEditor/Private/TextureEditorToolkit.cpp
@@ -811,13 +811,13 @@ void FTextureEditorToolkit::PopulateQuickInfo( )
 			bool bSourceAlphaDetected = FImageCore::DetectAlphaChannel(View);
 			SourceMipsAlphaDetectedText->SetText(FText::Format(NSLOCTEXT("TextureEditor", "QuickInfo_SourceAlphaDetected", "Source Alpha Detected: {0}"),
 				bSourceAlphaDetected ? NSLOCTEXT("TextureEditor", "True", "True") : NSLOCTEXT("TextureEditor", "False", "False")));
-			DetectSourceAlphaButton->SetVisibility(EVisibility::Hidden);
+			DetectSourceAlphaButton->SetVisibility(EVisibility::Collapsed);
 		}
 		else
 		{
 			SourceMipsAlphaDetectedText->SetText(FText::Format(NSLOCTEXT("TextureEditor", "QuickInfo_SourceAlphaDetected", "Source Alpha Detected: {0}"),
 				NSLOCTEXT("TextureEditor", "Unknown", "Unknown")));
-			DetectSourceAlphaButton->SetVisibility(Texture->Source.IsValid() ? EVisibility::Visible : EVisibility::Hidden);
+			DetectSourceAlphaButton->SetVisibility(Texture->Source.IsValid() ? EVisibility::Visible : EVisibility::Collapsed);
 		}
 	} // end if valid platform data
 
