Fix up materials after updating virtual texture streaming flag in texture editor.
Fix up can be quite slow on very large asset databases (we do through up a progress dialog). Will look at optimizing the speed of that next.

diff --git a/Engine/Source/Editor/TextureEditor/Private/Customizations/TextureDetailsCustomization.cpp b/Engine/Source/Editor/TextureEditor/Private/Customizations/TextureDetailsCustomization.cpp
--- a/Engine/Source/Editor/TextureEditor/Private/Customizations/TextureDetailsCustomization.cpp
+++ b/Engine/Source/Editor/TextureEditor/Private/Customizations/TextureDetailsCustomization.cpp
@@ -1,14 +1,18 @@
 // Copyright Epic Games, Inc. All Rights Reserved.
 
 #include "Customizations/TextureDetailsCustomization.h"
-#include "Misc/MessageDialog.h"
-#include "Widgets/DeclarativeSyntaxSupport.h"
-#include "Engine/Texture.h"
-#include "Editor.h"
+
+#include "AssetToolsModule.h"
+#include "DetailCategoryBuilder.h"
 #include "DetailLayoutBuilder.h"
 #include "DetailWidgetRow.h"
+#include "Editor.h"
+#include "Engine/Texture2D.h"
+#include "IAssetTools.h"
 #include "IDetailPropertyRow.h"
-#include "DetailCategoryBuilder.h"
+#include "Misc/MessageDialog.h"
+#include "Modules/ModuleManager.h"
+#include "Widgets/DeclarativeSyntaxSupport.h"
 #include "Widgets/Input/SNumericEntryBox.h"
 #include "Widgets/Input/SButton.h"
 
@@ -133,14 +137,58 @@ void FTextureDetails::CustomizeDetails(IDetailLayoutBuilder& DetailBuilder)
 			];
 	}
 
-	// Hide the option to enable VT streaming, if VT is disabled for the project
 	if (VirtualTextureStreamingPropertyHandle.IsValid())
 	{
-		static const auto CVarVirtualTexturesEnabled = IConsoleManager::Get().FindTConsoleVariableDataInt(TEXT("r.VirtualTextures")); check(CVarVirtualTexturesEnabled);
-		const bool bVirtualTextureEnabled = CVarVirtualTexturesEnabled->GetValueOnAnyThread() != 0;
-		if (!bVirtualTextureEnabled)
+		DetailBuilder.HideProperty(VirtualTextureStreamingPropertyHandle);
+
+		// Only show the option to enable VT streaming, if VT is enabled for the project.
+		static const auto CVarVirtualTexturesEnabled = IConsoleManager::Get().FindTConsoleVariableDataInt(TEXT("r.VirtualTextures"));
+		const static bool bVirtualTextureEnabled = CVarVirtualTexturesEnabled && CVarVirtualTexturesEnabled->GetValueOnAnyThread() != 0;
+
+		TArray<TWeakObjectPtr<UObject>> ObjectsBeingCustomized;
+		DetailBuilder.GetObjectsBeingCustomized(ObjectsBeingCustomized);
+		const bool bIsValidObject = ObjectsBeingCustomized.Num() > 0;
+
+		if (bVirtualTextureEnabled && bIsValidObject)
 		{
-			DetailBuilder.HideProperty(VirtualTextureStreamingPropertyHandle);
+			IDetailCategoryBuilder& TextureCategory = DetailBuilder.EditCategory("Texture");
+			TextureCategory.AddCustomRow(VirtualTextureStreamingPropertyHandle->GetPropertyDisplayName())
+			.NameContent()
+			[
+				VirtualTextureStreamingPropertyHandle->CreatePropertyNameWidget()
+			]
+			.ValueContent()
+			.VAlign(VAlign_Center)
+			[
+				SNew(SButton)
+				.VAlign(VAlign_Center)
+				.TextStyle(FAppStyle::Get(), "NormalText")
+				.Text_Lambda([WeakObject = ObjectsBeingCustomized[0]]()
+				{
+					UTexture2D const* Texture = Cast<UTexture2D>(WeakObject.Get());
+					const bool bVirtualTextureStreaming = Texture != nullptr ? Texture->VirtualTextureStreaming != 0 : false;
+					if (bVirtualTextureStreaming)
+					{
+						return LOCTEXT("Button_ConvertToRegularTexture", "Convert to Regular Texture");
+					}
+					else
+					{
+						return LOCTEXT("Button_ConvertToVirtualTexture", "Convert to Virtual Texture");
+					}
+				})
+				.OnClicked_Lambda([WeakObject = ObjectsBeingCustomized[0]]()
+				{
+					if (UTexture2D* Texture = Cast<UTexture2D>(WeakObject.Get()))
+					{
+						TArray<UTexture2D*> Textures;
+						Textures.Add(Texture);
+						IAssetTools& AssetTools = FModuleManager::LoadModuleChecked<FAssetToolsModule>("AssetTools").Get();
+						AssetTools.ConvertVirtualTextures(Textures, Texture->VirtualTextureStreaming);
+						return FReply::Handled();
+					}
+					return FReply::Unhandled();
+				})
+			];
 		}
 	}
 }
diff --git a/Engine/Source/Editor/VirtualTexturingEditor/Private/VirtualTexturingEditorModule.cpp b/Engine/Source/Editor/VirtualTexturingEditor/Private/VirtualTexturingEditorModule.cpp
--- a/Engine/Source/Editor/VirtualTexturingEditor/Private/VirtualTexturingEditorModule.cpp
+++ b/Engine/Source/Editor/VirtualTexturingEditor/Private/VirtualTexturingEditorModule.cpp
@@ -124,8 +124,9 @@ void FVirtualTexturingEditorModule::ConvertVirtualTextures(const TArray<UTexture
 {
 	FVirtualTextureConversionWorker VirtualTextureConversionWorker(bConvertBackToNonVirtual);
 	VirtualTextureConversionWorker.UserTextures = ObjectPtrWrap(Textures);
-	//We want all given texture to be added, so we put a minimum texture size of 0
-	VirtualTextureConversionWorker.FilterList(0);
+	// FilterList takes a threshold that is either greater or less than according to which direction we are converting in.
+	const int32 FilterListSizeThreshold = bConvertBackToNonVirtual ? MAX_int32 : 0;
+	VirtualTextureConversionWorker.FilterList(FilterListSizeThreshold);
 	if (RelatedMaterials)
 	{
 		VirtualTextureConversionWorker.Materials.Append(*RelatedMaterials);
