Show the pixel and texture formats that textures map to for all platforms.

diff --git a/Engine/Source/Editor/TextureEditor/Private/TextureEditorToolkit.cpp b/Engine/Source/Editor/TextureEditor/Private/TextureEditorToolkit.cpp
--- a/Engine/Source/Editor/TextureEditor/Private/TextureEditorToolkit.cpp
+++ b/Engine/Source/Editor/TextureEditor/Private/TextureEditorToolkit.cpp
@@ -58,6 +58,9 @@
 #include "Compression/OodleDataCompressionUtil.h"
 #include "Components/SceneCaptureComponent2D.h"
 #include "TextureResource.h"
+#include "TextureCompressorModule.h"
+#include "TextureBuildUtilities.h"
+#include "TextureFormatManager.h"
 
 #define LOCTEXT_NAMESPACE "FTextureEditorToolkit"
 
@@ -72,6 +75,14 @@ DEFINE_LOG_CATEGORY_STATIC(LogTextureEditor, Log, All);
 const FName FTextureEditorToolkit::ViewportTabId(TEXT("TextureEditor_Viewport"));
 const FName FTextureEditorToolkit::PropertiesTabId(TEXT("TextureEditor_Properties"));
 const FName FTextureEditorToolkit::OodleTabId(TEXT("TextureEditor_Oodle"));
+const FName FTextureEditorToolkit::FormatsTabId(TEXT("TextureEditor_Formats"));
+
+namespace
+{
+	const FName FormatsColumnId_Platform(TEXT("PlatformNameCol"));
+	const FName FormatsColumnId_TextureFormat(TEXT("TextureFormatNameCol"));
+	const FName FormatsColumnId_PixelFormat(TEXT("PixelFormatCol"));
+}
 
 UNREALED_API void GetBestFitForNumberOfTiles(int32 InSize, int32& OutRatioX, int32& OutRatioY);
 
@@ -238,6 +249,11 @@ void FTextureEditorToolkit::RegisterTabSpawners( const TSharedRef<class FTabMana
 		.SetDisplayName(LOCTEXT("OodleTab", "Oodle"))
 		.SetGroup(WorkspaceMenuCategoryRef)
 		.SetIcon(FSlateIcon(FAppStyle::GetAppStyleSetName(), "LevelEditor.Tabs.Details"));
+
+	InTabManager->RegisterTabSpawner(FormatsTabId, FOnSpawnTab::CreateSP(this, &FTextureEditorToolkit::HandleTabSpawnerSpawnFormats))
+		.SetDisplayName(LOCTEXT("FormatsTab", "Formats"))
+		.SetGroup(WorkspaceMenuCategoryRef)
+		.SetIcon(FSlateIcon(FAppStyle::GetAppStyleSetName(), "LevelEditor.Tabs.Details"));
 }
 
 
@@ -248,6 +264,7 @@ void FTextureEditorToolkit::UnregisterTabSpawners( const TSharedRef<class FTabMa
 	InTabManager->UnregisterTabSpawner(ViewportTabId);
 	InTabManager->UnregisterTabSpawner(PropertiesTabId);
 	InTabManager->UnregisterTabSpawner(OodleTabId);
+	InTabManager->UnregisterTabSpawner(FormatsTabId);
 }
 
 
@@ -343,7 +360,7 @@ void FTextureEditorToolkit::InitTextureEditor( const EToolkitMode::Type Mode, co
 	BindCommands();
 	CreateInternalWidgets();
 
-	const TSharedRef<FTabManager::FLayout> StandaloneDefaultLayout = FTabManager::NewLayout("Standalone_TextureEditor_Layout_v5")
+	const TSharedRef<FTabManager::FLayout> StandaloneDefaultLayout = FTabManager::NewLayout("Standalone_TextureEditor_Layout_v6")
 		->AddArea
 		(
 			FTabManager::NewPrimaryArea()
@@ -365,6 +382,7 @@ void FTextureEditorToolkit::InitTextureEditor( const EToolkitMode::Type Mode, co
 					FTabManager::NewStack()
 						->AddTab(PropertiesTabId, ETabState::OpenedTab)
 						->AddTab(OodleTabId, ETabState::OpenedTab)
+						->AddTab(FormatsTabId, ETabState::OpenedTab)
 						->SetForegroundTab(PropertiesTabId)
 						->SetSizeCoefficient(0.33f)
 				)
@@ -662,6 +680,46 @@ void FTextureEditorToolkit::PopulateQuickInfo( )
 					OutstandingEstimation = PlatformData->LaunchEstimateOnDiskSizeTask(OodleCompressor, OodleCompressionLevel, CompressionBlockSize, Texture->GetPathName());
 				}
 
+				// We only want to refresh the platform list when we have a settings change (AlreadyHaveResults, above)				
+				{
+					PlatformTextureFormatList.Empty();
+
+					ITextureFormatManagerModule* TFMM = FModuleManager::GetModulePtr<ITextureFormatManagerModule>("TextureFormat");
+					ITargetPlatformManagerModule* TPM = GetTargetPlatformManager();
+					if (TPM)
+					{
+						for (FName PlatformName : AvailablePlatformNames)
+						{
+							ITargetPlatform* TargetPlatform = TPM->FindTargetPlatform(PlatformName);
+							if (TargetPlatform)
+							{
+								TArray<TArray<FTextureBuildSettings>> Settings;
+								Texture->GetTargetPlatformBuildSettings(TargetPlatform, Settings);
+
+								FTextureBuildSettings& BuildSettings = Settings[0][0];
+								FName TextureFormatName = UE::TextureBuildUtilities::TextureFormatRemovePlatformPrefixFromName(BuildSettings.TextureFormatName);
+					
+								ITextureFormat const* TextureFormat = TFMM->FindTextureFormat(BuildSettings.TextureFormatName);
+								if (TextureFormat)
+								{
+									TSharedPtr<FPlatformTextureFormats> Formats = MakeShared<FPlatformTextureFormats>();
+									Formats->PlatformName = PlatformName;
+									Formats->TextureFormatName = TextureFormatName;
+									Formats->PixelFormat_AlphaTrue = TextureFormat->GetEncodedPixelFormat(BuildSettings, true);
+									Formats->PixelFormat_AlphaFalse = TextureFormat->GetEncodedPixelFormat(BuildSettings, false);
+									Formats->bKnowAlphaTransparency = BuildSettings.bKnowAlphaTransparency;
+									Formats->bHasTransparentAlpha = BuildSettings.bHasTransparentAlpha;
+
+									PlatformTextureFormatList.Add(Formats);
+								}
+							}
+						}
+					}
+
+					PlatformTextureFormatsWidget->RebuildList();
+				}
+
+
 				OodleCompressedPreviewDDCKey = PlatformDataPtr[0]->DerivedDataKey;
 			}
 
@@ -1334,6 +1392,68 @@ TSharedRef<SWidget> FTextureEditorToolkit::BuildTexturePropertiesWidget( )
 	return TexturePropertiesWidget.ToSharedRef();
 }
 
+// Manages the slate widgets for a single entry in the list of platforms and their formats.
+class SPlatformTextureFormatRow : public SMultiColumnTableRow<TSharedRef<FTextureEditorToolkit::FPlatformTextureFormats>>
+{
+public:
+	SLATE_BEGIN_ARGS(SPlatformTextureFormatRow) {}
+	SLATE_END_ARGS()
+
+	void Construct(const FArguments& InArgs, const TSharedRef<STableViewBase>& InOwnerTable, const TSharedPtr<FTextureEditorToolkit::FPlatformTextureFormats>& InNode)
+	{
+		Node = InNode;
+
+		FSuperRowType::FArguments StyleArguments = FSuperRowType::FArguments()
+		.Padding(1.0f)
+		.Style(&FAppStyle().GetWidgetStyle<FTableRowStyle>("TableView.AlternatingRow"));
+
+		SMultiColumnTableRow::Construct(StyleArguments, InOwnerTable);
+	}
+
+	virtual TSharedRef<SWidget> GenerateWidgetForColumn(const FName& InColumnName) override
+	{
+		if (InColumnName == FormatsColumnId_Platform)
+		{
+			return SNew(STextBlock).Text(FText::FromName(Node->PlatformName)).Margin(FMargin(2, 0, 2, 0));
+		}
+		else if (InColumnName == FormatsColumnId_TextureFormat)
+		{
+			return SNew(STextBlock).Text(FText::FromName(Node->TextureFormatName)).Margin(FMargin(2, 0, 2, 0));
+		}
+		else
+		{
+			// Auto* formats will automatically use an opaque pixel format if not needed. We might or might not know
+			// at this point which will get chosen - and it might not matter.
+			if (Node->PixelFormat_AlphaTrue == Node->PixelFormat_AlphaFalse)
+			{
+				return SNew(STextBlock).Text(FText::AsCultureInvariant(GetPixelFormatString(Node->PixelFormat_AlphaFalse))).Margin(FMargin(2, 0, 2, 0));
+			}
+
+			if (Node->bKnowAlphaTransparency)
+			{
+				EPixelFormat PF = Node->PixelFormat_AlphaFalse;
+				if (Node->bHasTransparentAlpha)
+				{
+					PF = Node->PixelFormat_AlphaTrue;
+				}
+				return SNew(STextBlock).Text(FText::AsCultureInvariant(GetPixelFormatString(PF))).Margin(FMargin(2, 0, 2, 0));
+			}
+
+			// If we don't know, show both. They can resolve this by clicking "Scan 
+			TStringBuilder<128> Str;
+			Str << GetPixelFormatString(Node->PixelFormat_AlphaFalse) << " or " << GetPixelFormatString(Node->PixelFormat_AlphaTrue);
+			return SNew(STextBlock).Text(FText::AsCultureInvariant(Str.ToString())).Margin(FMargin(2, 0, 2, 0));
+		}
+	}
+
+	TSharedPtr<FTextureEditorToolkit::FPlatformTextureFormats> Node;
+};
+
+TSharedRef<ITableRow> FTextureEditorToolkit::MakePlatformTextureFormatWidget(TSharedPtr<FPlatformTextureFormats> PlatformTextureFormat, const TSharedRef<STableViewBase>& OwnerTable)
+{
+	return SNew(SPlatformTextureFormatRow, OwnerTable, PlatformTextureFormat);
+}
+
 void FTextureEditorToolkit::CreateInternalWidgets()
 {
 	//
@@ -1407,6 +1527,30 @@ void FTextureEditorToolkit::CreateInternalWidgets()
 
 	TextureViewport = SNew(STextureEditorViewport, SharedThis(this));
 
+	FormatsTabContainer = SNew(SVerticalBox)
+	+ SVerticalBox::Slot()
+	.AutoHeight()
+	.Padding(4.0f)
+	[
+		SAssignNew(PlatformTextureFormatsWidget, SListView<TSharedPtr<FPlatformTextureFormats>>)
+		.ListItemsSource(&PlatformTextureFormatList)
+		.SelectionMode(ESelectionMode::None)
+		.HeaderRow(
+			SNew(SHeaderRow)
+			+SHeaderRow::Column(FormatsColumnId_Platform)
+			.DefaultLabel(LOCTEXT("TextureEditor_PlatformNameCol", "Platform"))
+			.FillWidth(.3f)
+			+SHeaderRow::Column(FormatsColumnId_TextureFormat)
+			.DefaultLabel(LOCTEXT("TextureEditor_TextureFormatCol", "Texture Format"))
+			.FillWidth(.4f)
+			+SHeaderRow::Column(FormatsColumnId_PixelFormat)
+			.DefaultLabel(LOCTEXT("TextureEditor_PixelFormatCol", "Pixel Format"))
+			.DefaultTooltip(LOCTEXT("TextureEditor_PixelFormatColTT", "The pixel format that will be chosen for this platform. For texture formats that automatically select alpha or not (e.g. AutoDXT) this might not be known. Click \"Detect\" in the details panel to run the alpha scan and get the actual format"))
+			.FillWidth(.3f)
+		)
+		.OnGenerateRow(this, &FTextureEditorToolkit::MakePlatformTextureFormatWidget)
+	]; // end format details.
+
 	OodleTabContainer = SNew(SVerticalBox)
 
 	//
@@ -2710,6 +2854,20 @@ void FTextureEditorToolkit::HandleSettingsActionExecute( )
 	FModuleManager::LoadModuleChecked<ISettingsModule>("Settings").ShowViewer("Editor", "ContentEditors", "TextureEditor");
 }
 
+
+TSharedRef<SDockTab> FTextureEditorToolkit::HandleTabSpawnerSpawnFormats(const FSpawnTabArgs& Args)
+{
+	check(Args.GetTabId() == FormatsTabId);
+
+	TSharedRef<SDockTab> SpawnedTab = SNew(SDockTab)
+	.Label(LOCTEXT("TextureFormatsTitle", "Formats"))
+	[
+		FormatsTabContainer.ToSharedRef()
+	];
+
+	return SpawnedTab;
+}
+
 TSharedRef<SDockTab> FTextureEditorToolkit::HandleTabSpawnerSpawnOodle(const FSpawnTabArgs& Args)
 {
 	check(Args.GetTabId() == OodleTabId);
@@ -3588,6 +3746,8 @@ FReply FTextureEditorToolkit::DetectSourceAlphaButton_Clicked()
 	Texture->PreEditChange(nullptr);
 	Texture->Source.UpdateChannelLinearMinMax();
 	Texture->PostEditChange();
+	OodleCompressedPreviewDDCKey.Set<FString>(FString());
+	PopulateQuickInfo();
 	return FReply::Handled();
 }
 
diff --git a/Engine/Source/Editor/TextureEditor/Private/TextureEditorToolkit.h b/Engine/Source/Editor/TextureEditor/Private/TextureEditorToolkit.h
--- a/Engine/Source/Editor/TextureEditor/Private/TextureEditorToolkit.h
+++ b/Engine/Source/Editor/TextureEditor/Private/TextureEditorToolkit.h
@@ -13,6 +13,7 @@
 #include "IDetailsView.h"
 #include "TextureEditorSettings.h"
 
+template <typename ItemType> class SListView;
 class SMenuAnchor;
 class STextBlock;
 class SButton;
@@ -333,6 +334,7 @@ private:
 	TSharedRef<SDockTab> HandleTabSpawnerSpawnProperties(const FSpawnTabArgs& Args);
 	TSharedRef<SDockTab> HandleTabSpawnerSpawnViewport( const FSpawnTabArgs& Args );
 	TSharedRef<SDockTab> HandleTabSpawnerSpawnOodle(const FSpawnTabArgs& Args);
+	TSharedRef<SDockTab> HandleTabSpawnerSpawnFormats(const FSpawnTabArgs& Args);
 
 	// Callback for toggling the Texture Border action.
 	void HandleTextureBorderActionExecute( );
@@ -402,6 +404,7 @@ private:
 
 	/** Oodle tab */
 	TSharedPtr<SVerticalBox> OodleTabContainer;
+	TSharedPtr<SVerticalBox> FormatsTabContainer;
 
 	/** Properties tab */
 	TSharedPtr<SVerticalBox> TextureProperties;
@@ -514,6 +517,28 @@ private:
 	TArray<TSharedPtr<FString>> AvailablePlatforms;
 	TArray<FName> AvailablePlatformNames;
 
+public:
+	struct FPlatformTextureFormats
+	{
+		FName PlatformName;
+		FName TextureFormatName;
+
+		// Pixel formats for whether we detected alpha in the texture source. Only
+		// matters for Auto* texture formats. If we don't know the texture source alpha
+		// information we show both.
+		EPixelFormat PixelFormat_AlphaTrue;
+		EPixelFormat PixelFormat_AlphaFalse;
+		bool bKnowAlphaTransparency;
+		bool bHasTransparentAlpha;
+	};
+private:
+
+	TSharedPtr<SListView<TSharedPtr<FPlatformTextureFormats>>> PlatformTextureFormatsWidget;
+	TArray<TSharedPtr<FPlatformTextureFormats>> PlatformTextureFormatList;
+	TSharedRef<ITableRow> MakePlatformTextureFormatWidget(TSharedPtr<FPlatformTextureFormats> PlatformTextureFormat, const TSharedRef<STableViewBase>& OwnerTable);
+
+	TMap<FName, TPair<TSharedPtr<STextBlock>, TSharedPtr<STextBlock>>> PlatformFormatDetails;
+
 	// Objects and callbacks for the custom encoding settings checkbox
 	TSharedPtr<class SCheckBox> OodleOverrideCheck;
 	void OnUseEditorOodleSettingsChanged(ECheckBoxState NewState);
@@ -581,4 +606,5 @@ private:
 	static const FName ViewportTabId;
 	static const FName PropertiesTabId;
 	static const FName OodleTabId;
+	static const FName FormatsTabId;
 };
diff --git a/Engine/Source/Editor/TextureEditor/TextureEditor.Build.cs b/Engine/Source/Editor/TextureEditor/TextureEditor.Build.cs
--- a/Engine/Source/Editor/TextureEditor/TextureEditor.Build.cs
+++ b/Engine/Source/Editor/TextureEditor/TextureEditor.Build.cs
@@ -38,7 +38,9 @@ public class TextureEditor : ModuleRules
 				"MediaAssets",
 				"DerivedDataCache",
 				"DesktopPlatform",
-				"DeveloperToolSettings"
+				"DeveloperToolSettings",
+				"TextureCompressor",
+				"TextureBuildUtilities"
 			}
 		);
 	}
