Showing mip size details per platform in texture viewer.

diff --git a/Engine/Source/Developer/TextureFormat/Public/Interfaces/ITextureFormat.h b/Engine/Source/Developer/TextureFormat/Public/Interfaces/ITextureFormat.h
--- a/Engine/Source/Developer/TextureFormat/Public/Interfaces/ITextureFormat.h
+++ b/Engine/Source/Developer/TextureFormat/Public/Interfaces/ITextureFormat.h
@@ -108,6 +108,9 @@ struct FEncodedTextureExtendedData
 * "Inline" textures are loaded with the texture asset, and streaming textures are loaded on demand.
 * Generally, 7 of the smallest mips are inlined, however some platforms pack a lot of mips in to a single
 * allocation ("packed mip tail" = NumMipsInTail). Those mips must all be inlined.
+*
+* This function is used primarily for determining where to put mips in the DDC. Runtime inlining is subject
+* to further constraints and must use the GetNumNonStreamingMipsDirect() function.
 * 
 * InNumMips			    The total mips that the texture contains.
 * InExtendedData	    If the texture is being built for a platform that provides extended data, pass it
@@ -141,6 +144,56 @@ static int32 GetNumStreamingMipsDirect(int32 InNumMips, bool bInCubeMap, bool bI
 	return NumStreamingMips;
 }
 
+/**
+	Return how many mips need to be inlined for the given texture properties. This pays attention to runtime
+	size restrictions and should be used for serializing cooked data.
+
+	InNumMips, InTopMipSizeX, InTopMipSizeY must all be sizes _before_ any LOD stripping during cook, with the expectation
+	that the returned value is verified against the remaining mip count after LOD stripping.
+
+	This function is not valid if streaming is disabled for the texture.
+*/
+static int32 GetNumNonStreamingMipsDirect(int32 InNumMips, int32 InTopMipSizeX, int32 InTopMipSizeY, EPixelFormat InPixelFormat, int32 InNumMipsInTail, int32 InNumInlineDerivedMips)
+{
+	if (InNumMips <= 1)
+	{
+		return InNumMips;
+	}
+
+	int32 NumNonStreamingMips = 1;
+
+	// Take the min resident limit into account.
+	if (InNumMipsInTail)
+	{
+		// If we have a tail, it can't be broken up so the whole thing must be non streaming.
+		NumNonStreamingMips = FMath::Max(NumNonStreamingMips, InNumMipsInTail);
+	}
+	NumNonStreamingMips = FMath::Max(NumNonStreamingMips, InNumInlineDerivedMips);
+	NumNonStreamingMips = FMath::Min(NumNonStreamingMips, InNumMips);
+
+	if ( RequiresBlock4Alignment(InPixelFormat) )
+	{
+		// ensure the top non-streamed mip (and all streamed mips) size is >= BlockSize (and a multiple of block size!)
+			
+		// note: this is not right for non pow 2; NeverStream should set !bIsStreamingPossible in that case
+		if ( FMath::IsPowerOfTwo(InTopMipSizeX) && FMath::IsPowerOfTwo(InTopMipSizeY) )
+		{
+			// This is not right for textures smaller than 4x4, however such textures are already non-streaming due to
+			// NumInlineDerivedMips (which is functionally hard coded to 7).
+			NumNonStreamingMips = FMath::Max<int32>(NumNonStreamingMips, InNumMips - FPlatformMath::FloorLog2(InTopMipSizeX / 4));
+			NumNonStreamingMips = FMath::Max<int32>(NumNonStreamingMips, InNumMips - FPlatformMath::FloorLog2(InTopMipSizeY / 4));
+		}
+		else
+		{
+			// should never get here as streaming should be disabled in this case, so we should not hit this branch, but it's not reliable
+			NumNonStreamingMips = InNumMips;
+		}
+	}
+
+	return NumNonStreamingMips;
+}
+
+
 // Everything necessary to know the memory layout for an encoded untiled unpacked texture (i.e. enough information
 // to describe the texture entirely to a PC hardware API).
 // Once a texture gets tiled or gets a packed mip tail, FEncodedTextureEncodedData is additionally
@@ -280,11 +333,31 @@ struct FEncodedTextureDescription
 		return WidthInBlocks * GPixelFormats[PixelFormat].BlockBytes;
 	}
 
+	// This is the number of streaming mips without any restrictions associated with block size alignment. It's used
+	// primarily for where mips go in the DDC.
 	int32 GetNumStreamingMips(const FEncodedTextureExtendedData* InExtendedData, const FTextureEngineParameters& InEngineParameters) const
 	{
 		return GetNumStreamingMipsDirect(NumMips, bCubeMap, bVolumeTexture, bTextureArray, InExtendedData, InEngineParameters);
 	}
 
+	// This is the number of nonstreaming mips that works for cooking and thus is aware of block alignment concerns for runtime use.
+	// Only call this if streaming is possible.
+	int32 GetNumNonStreamingMips(const FEncodedTextureExtendedData* InExtendedData, const FTextureEngineParameters& InEngineParameters)
+	{
+		return GetNumNonStreamingMipsDirect(NumMips, TopMipSizeX, TopMipSizeY, PixelFormat, InExtendedData ? InExtendedData->NumMipsInTail : 0, InEngineParameters.NumInlineDerivedMips);
+	}
+
+	// Fills out the given array with linear mip sizes for the number of mips this texture has. These sizes are only valid
+	// if the texture is linear and has no packed mip tail (i.e. no associated FEncodedTextureExtendedData).
+	void GetLinearMipSizesInBytes(TArray<uint64, TInlineAllocator<FEncodedTextureExtendedData::MAX_TEXTURE_MIP_COUNT>>& OutMipSizesInBytes) const
+	{
+		OutMipSizesInBytes.SetNumUninitialized(NumMips);
+		for (int32 MipIndex = 0; MipIndex < NumMips; MipIndex++)
+		{
+			OutMipSizesInBytes[MipIndex] = GetMipSizeInBytes(MipIndex);
+		}
+	}
+
 	// Convenience function for iterating over the encoded mips when you need to know how many mips are represented. Use as:
 	//
 	//	for (int32 EncodedMipIndex = 0; EncodedMipIndex < OutMipTailIndex + 1; EncodedMipIndex++)
@@ -534,11 +607,7 @@ public:
 	[[nodiscard]] virtual FEncodedTextureExtendedData GetExtendedDataForTexture(const FEncodedTextureDescription& InTextureDescription, int8 InLODBias) const
 	{
 		FEncodedTextureExtendedData ExtendedData;
-		ExtendedData.MipSizesInBytes.SetNumUninitialized(InTextureDescription.NumMips);
-		for (int32 MipIndex = 0; MipIndex < InTextureDescription.NumMips; MipIndex++)
-		{
-			ExtendedData.MipSizesInBytes[MipIndex] = InTextureDescription.GetMipSizeInBytes(MipIndex);
-		}
+		InTextureDescription.GetLinearMipSizesInBytes(ExtendedData.MipSizesInBytes);
 		return ExtendedData;
 	}
 
diff --git a/Engine/Source/Editor/TextureEditor/Private/TextureEditorToolkit.cpp b/Engine/Source/Editor/TextureEditor/Private/TextureEditorToolkit.cpp
--- a/Engine/Source/Editor/TextureEditor/Private/TextureEditorToolkit.cpp
+++ b/Engine/Source/Editor/TextureEditor/Private/TextureEditorToolkit.cpp
@@ -75,13 +75,19 @@ DEFINE_LOG_CATEGORY_STATIC(LogTextureEditor, Log, All);
 const FName FTextureEditorToolkit::ViewportTabId(TEXT("TextureEditor_Viewport"));
 const FName FTextureEditorToolkit::PropertiesTabId(TEXT("TextureEditor_Properties"));
 const FName FTextureEditorToolkit::OodleTabId(TEXT("TextureEditor_Oodle"));
-const FName FTextureEditorToolkit::FormatsTabId(TEXT("TextureEditor_Formats"));
+const FName FTextureEditorToolkit::PlatformsTabId(TEXT("TextureEditor_Platforms"));
 
 namespace
 {
-	const FName FormatsColumnId_Platform(TEXT("PlatformNameCol"));
-	const FName FormatsColumnId_TextureFormat(TEXT("TextureFormatNameCol"));
-	const FName FormatsColumnId_PixelFormat(TEXT("PixelFormatCol"));
+	const FName PlatformsColumnId_Platform(TEXT("PlatformNameCol"));
+	const FName PlatformsColumnId_TextureFormat(TEXT("TextureFormatNameCol"));
+	const FName PlatformsColumnId_PixelFormat(TEXT("PixelFormatCol"));
+	const FName PlatformsColumnId_LODBias(TEXT("LODBiasCol"));
+	const FName PlatformsColumnId_Dims(TEXT("DimsCol"));
+	const FName PlatformsColumnId_NumMips(TEXT("NumMipsCol"));
+	const FName PlatformsColumnId_SizeOptional(TEXT("MipSizesCol_Optional"));
+	const FName PlatformsColumnId_SizeStreaming(TEXT("MipSizesCol_Streaming"));
+	const FName PlatformsColumnId_SizeInline(TEXT("MipSizesCol_Inline"));
 }
 
 UNREALED_API void GetBestFitForNumberOfTiles(int32 InSize, int32& OutRatioX, int32& OutRatioY);
@@ -250,8 +256,8 @@ void FTextureEditorToolkit::RegisterTabSpawners( const TSharedRef<class FTabMana
 		.SetGroup(WorkspaceMenuCategoryRef)
 		.SetIcon(FSlateIcon(FAppStyle::GetAppStyleSetName(), "LevelEditor.Tabs.Details"));
 
-	InTabManager->RegisterTabSpawner(FormatsTabId, FOnSpawnTab::CreateSP(this, &FTextureEditorToolkit::HandleTabSpawnerSpawnFormats))
-		.SetDisplayName(LOCTEXT("FormatsTab", "Formats"))
+	InTabManager->RegisterTabSpawner(PlatformsTabId, FOnSpawnTab::CreateSP(this, &FTextureEditorToolkit::HandleTabSpawnerSpawnPlatforms))
+		.SetDisplayName(LOCTEXT("PlatformsTab", "Platforms"))
 		.SetGroup(WorkspaceMenuCategoryRef)
 		.SetIcon(FSlateIcon(FAppStyle::GetAppStyleSetName(), "LevelEditor.Tabs.Details"));
 }
@@ -264,7 +270,7 @@ void FTextureEditorToolkit::UnregisterTabSpawners( const TSharedRef<class FTabMa
 	InTabManager->UnregisterTabSpawner(ViewportTabId);
 	InTabManager->UnregisterTabSpawner(PropertiesTabId);
 	InTabManager->UnregisterTabSpawner(OodleTabId);
-	InTabManager->UnregisterTabSpawner(FormatsTabId);
+	InTabManager->UnregisterTabSpawner(PlatformsTabId);
 }
 
 
@@ -360,7 +366,7 @@ void FTextureEditorToolkit::InitTextureEditor( const EToolkitMode::Type Mode, co
 	BindCommands();
 	CreateInternalWidgets();
 
-	const TSharedRef<FTabManager::FLayout> StandaloneDefaultLayout = FTabManager::NewLayout("Standalone_TextureEditor_Layout_v6")
+	const TSharedRef<FTabManager::FLayout> StandaloneDefaultLayout = FTabManager::NewLayout("Standalone_TextureEditor_Layout_v7")
 		->AddArea
 		(
 			FTabManager::NewPrimaryArea()
@@ -382,7 +388,7 @@ void FTextureEditorToolkit::InitTextureEditor( const EToolkitMode::Type Mode, co
 					FTabManager::NewStack()
 						->AddTab(PropertiesTabId, ETabState::OpenedTab)
 						->AddTab(OodleTabId, ETabState::OpenedTab)
-						->AddTab(FormatsTabId, ETabState::OpenedTab)
+						->AddTab(PlatformsTabId, ETabState::OpenedTab)
 						->SetForegroundTab(PropertiesTabId)
 						->SetSizeCoefficient(0.33f)
 				)
@@ -710,6 +716,99 @@ void FTextureEditorToolkit::PopulateQuickInfo( )
 									Formats->bKnowAlphaTransparency = BuildSettings.bKnowAlphaTransparency;
 									Formats->bHasTransparentAlpha = BuildSettings.bHasTransparentAlpha;
 
+									// Technically this texture might not be VT; this means that the user wants it to, but there are
+									// reasons why it might not.
+									if (Texture->VirtualTextureStreaming)
+									{
+										// Our mip sizes don't really makes sense and the VT mip count isn't readily available, so we
+										// don't report in these cases.
+										Formats->bIsVirtual = true;
+									}
+									else
+									{
+										Formats->bIsVirtual = false;
+
+										FTextureEngineParameters EngineParameters;
+										FEncodedTextureDescription TextureDescription;
+										FEncodedTextureExtendedData TextureExtendedData;
+										bool bHaveEncodedTextureDescription = BuildSettings.GetEncodedTextureDescriptionFromSourceMips(
+											&TextureDescription, TextureFormat, 
+											Texture->Source.GetSizeX(), Texture->Source.GetSizeY(), Texture->Source.GetNumSlices(),
+											Texture->Source.GetNumMips(), Formats->bHasTransparentAlpha);
+
+										if (bHaveEncodedTextureDescription)
+										{
+											if (BuildSettings.TilerEvenIfNotSharedLinear)
+											{
+												TextureExtendedData = BuildSettings.TilerEvenIfNotSharedLinear->GetExtendedDataForTexture(TextureDescription, BuildSettings.LODBias);
+											}
+											else
+											{
+												TextureDescription.GetLinearMipSizesInBytes(TextureExtendedData.MipSizesInBytes);
+											}
+
+											const UTextureLODSettings* LODSettings = (UTextureLODSettings*)UDeviceProfileManager::Get().FindProfile(TargetPlatform->PlatformName());
+											int32 FirstMipToSerialize = LODSettings->CalculateLODBias(
+												TextureDescription.TopMipSizeX, TextureDescription.TopMipSizeY, Texture->MaxTextureSize, 
+												Texture->LODGroup, Texture->LODBias, 0, Texture->MipGenSettings, false);
+
+											int32 StrippedNumMips = TextureDescription.NumMips;
+											const int32 LastMip = FMath::Max(StrippedNumMips - 1, 0);
+											const int32 FirstMipTailMip = StrippedNumMips - TextureExtendedData.NumMipsInTail;
+											check(FirstMipTailMip >= 0);
+											FirstMipToSerialize = FMath::Clamp(FirstMipToSerialize, 0, TextureExtendedData.NumMipsInTail > 0 ? FirstMipTailMip : LastMip);
+											StrippedNumMips = FMath::Max(0, StrippedNumMips - FirstMipToSerialize);
+											
+
+											int32 NumNonStreamingMips = TextureDescription.GetNumNonStreamingMips(&TextureExtendedData, EngineParameters);
+											NumNonStreamingMips = FMath::Min(NumNonStreamingMips, StrippedNumMips);
+											int32 FirstInlineMip = StrippedNumMips - NumNonStreamingMips;
+											int32 OptionalMipCount = 0;
+
+											if (StrippedNumMips > 0 )
+											{
+												const int32 FirstMipWidth  = TextureDescription.GetMipWidth(FirstMipToSerialize);
+												const int32 FirstMipHeight = TextureDescription.GetMipHeight(FirstMipToSerialize);
+
+												OptionalMipCount = LODSettings->CalculateNumOptionalMips(Texture->LODGroup, FirstMipWidth, FirstMipHeight, StrippedNumMips, FirstInlineMip, Texture->MipGenSettings);
+
+												// OptionalMips must be streaming mips.
+												check(OptionalMipCount <= FirstInlineMip);
+											}
+
+											Formats->LODBias = FirstMipToSerialize;
+											Formats->NumMips_TotalAfterLOD = TextureDescription.NumMips - FirstMipToSerialize;
+											Formats->Width = TextureDescription.GetMipWidth(FirstMipToSerialize);
+											Formats->Height = TextureDescription.GetMipHeight(FirstMipToSerialize);
+											Formats->NumMips_Tail = TextureExtendedData.NumMipsInTail;
+											
+											// We need to make sure all these indices have been adjusted for LODBias.
+											const int32 OptionalMipsEndAt = Formats->LODBias + OptionalMipCount;
+											const int32 InlineMipsStartAt = Formats->LODBias + FirstInlineMip;
+											const int32 NumAddressableMips = TextureDescription.NumMips - ((TextureExtendedData.NumMipsInTail > 0) ? (TextureExtendedData.NumMipsInTail - 1) : 0);
+										
+											int32 MipIndex = Formats->LODBias;
+											for (; MipIndex < OptionalMipsEndAt; MipIndex++)
+											{
+												Formats->MipChainSizeInBytes_Optional += TextureExtendedData.MipSizesInBytes[MipIndex];
+												Formats->NumMips_Optional++;												
+											}
+
+											for (; MipIndex < InlineMipsStartAt; MipIndex++)
+											{
+												Formats->NumMips_Streaming++;
+												Formats->MipChainSizeInBytes_Streaming += TextureExtendedData.MipSizesInBytes[MipIndex];
+											}
+											
+											for (; MipIndex < NumAddressableMips; MipIndex++)
+											{
+												Formats->MipChainSizeInBytes_Inline += TextureExtendedData.MipSizesInBytes[MipIndex];
+											}
+
+											Formats->MipSizesInBytes = MoveTemp(TextureExtendedData.MipSizesInBytes);
+										}
+									}
+
 									PlatformTextureFormatList.Add(Formats);
 								}
 							}
@@ -1410,16 +1509,127 @@ public:
 		SMultiColumnTableRow::Construct(StyleArguments, InOwnerTable);
 	}
 
+	FText GenerateMipSizeToolTip() const
+	{
+		TStringBuilder<64> MipSizes;
+		int32 Count = 0;
+		for (uint64 MipSize : Node->MipSizesInBytes)
+		{
+			if (Count)
+			{
+				MipSizes << TEXT('\n');
+			}
+			Count++;
+
+			if (Node->NumMips_Tail && Count == Node->MipSizesInBytes.Num())
+			{
+				MipSizes << Count << TEXTVIEW("..") << Count + Node->NumMips_Tail << TEXTVIEW(" - ");
+			}
+			else
+			{
+				MipSizes << Count << TEXTVIEW(" - ");
+			}
+
+			MipSizes << FText::AsNumber(MipSize).ToString();
+
+			if (Count <= Node->LODBias)
+			{
+				MipSizes << TEXTVIEW(" X");
+			}
+			else if (Count <= Node->LODBias + Node->NumMips_Optional)
+			{
+				MipSizes << TEXTVIEW(" O");
+			}
+			else if (Count <= Node->LODBias + Node->NumMips_Optional + Node->NumMips_Streaming)
+			{
+				MipSizes << TEXTVIEW(" S");
+			}
+			else
+			{
+				MipSizes << TEXTVIEW(" I");
+			}
+		}
+
+		return FText::AsCultureInvariant(MipSizes.ToString());
+	}
+
 	virtual TSharedRef<SWidget> GenerateWidgetForColumn(const FName& InColumnName) override
 	{
-		if (InColumnName == FormatsColumnId_Platform)
+		if (InColumnName == PlatformsColumnId_Platform)
 		{
 			return SNew(STextBlock).Text(FText::FromName(Node->PlatformName)).Margin(FMargin(2, 0, 2, 0));
 		}
-		else if (InColumnName == FormatsColumnId_TextureFormat)
+		else if (InColumnName == PlatformsColumnId_TextureFormat)
 		{
 			return SNew(STextBlock).Text(FText::FromName(Node->TextureFormatName)).Margin(FMargin(2, 0, 2, 0));
 		}
+		else if (InColumnName == PlatformsColumnId_LODBias)
+		{
+			if (Node->bIsVirtual)
+			{
+				return SNew(STextBlock);
+			}
+			return SNew(STextBlock).Text(FText::AsNumber(Node->LODBias)).Margin(FMargin(2, 0, 2, 0)).Justification(ETextJustify::Right);
+		}
+		else if (InColumnName == PlatformsColumnId_NumMips)
+		{
+			if (Node->bIsVirtual)
+			{
+				return SNew(STextBlock);
+			}
+			TStringBuilder<32> Mips;
+			Mips.Appendf(TEXT("%d / %d / %d / %d"), Node->NumMips_TotalAfterLOD, Node->NumMips_Optional, Node->NumMips_Streaming, Node->NumMips_Tail);
+			return SNew(STextBlock)
+				.Text(FText::AsCultureInvariant(Mips.ToString()))
+				.ToolTipText(NSLOCTEXT("TextureEditor", "MipInfoTT", "Total After LODBias / Optional / Streaming / PackedTail"))
+				.Margin(FMargin(2, 0, 2, 0));
+		}
+		else if (InColumnName == PlatformsColumnId_Dims)
+		{
+			if (Node->bIsVirtual)
+			{
+				return SNew(STextBlock).Text(NSLOCTEXT("TextureEditor", "CantBecauseVT", "N/A - VirtualTexture"));
+			}
+			TStringBuilder<32> Dims;
+			Dims << Node->Width << TEXT('x') << Node->Height;
+			return SNew(STextBlock).Text(FText::AsCultureInvariant(Dims.ToString())).Margin(FMargin(2, 0, 2, 0));
+		}
+		else if (InColumnName == PlatformsColumnId_SizeOptional)
+		{
+			if (Node->bIsVirtual)
+			{
+				return SNew(STextBlock);
+			}
+			return SNew(STextBlock)
+				.Text(FText::AsNumber(Node->MipChainSizeInBytes_Optional))
+				.ToolTipText(GenerateMipSizeToolTip())
+				.Margin(FMargin(2, 0, 2, 0))
+				.Justification(ETextJustify::Right);
+		}
+		else if (InColumnName == PlatformsColumnId_SizeInline)
+		{
+			if (Node->bIsVirtual)
+			{
+				return SNew(STextBlock);
+			}
+			return SNew(STextBlock)
+				.Text(FText::AsNumber(Node->MipChainSizeInBytes_Inline))
+				.ToolTipText(GenerateMipSizeToolTip())
+				.Margin(FMargin(2, 0, 2, 0))
+				.Justification(ETextJustify::Right);
+		}
+		else if (InColumnName == PlatformsColumnId_SizeStreaming)
+		{
+			if (Node->bIsVirtual)
+			{
+				return SNew(STextBlock);
+			}
+			return SNew(STextBlock)
+				.Text(FText::AsNumber(Node->MipChainSizeInBytes_Streaming))
+				.ToolTipText(GenerateMipSizeToolTip())
+				.Margin(FMargin(2, 0, 2, 0))
+				.Justification(ETextJustify::Right);
+		}
 		else
 		{
 			// Auto* formats will automatically use an opaque pixel format if not needed. We might or might not know
@@ -1439,10 +1649,13 @@ public:
 				return SNew(STextBlock).Text(FText::AsCultureInvariant(GetPixelFormatString(PF))).Margin(FMargin(2, 0, 2, 0));
 			}
 
-			// If we don't know, show both. They can resolve this by clicking "Scan 
+			// If we don't know, show both. They can resolve this by clicking "Detect" in the "Details" page.
 			TStringBuilder<128> Str;
 			Str << GetPixelFormatString(Node->PixelFormat_AlphaFalse) << " or " << GetPixelFormatString(Node->PixelFormat_AlphaTrue);
-			return SNew(STextBlock).Text(FText::AsCultureInvariant(Str.ToString())).Margin(FMargin(2, 0, 2, 0));
+			return SNew(STextBlock)
+				.Text(FText::AsCultureInvariant(Str.ToString()))
+				.ToolTipText(NSLOCTEXT("TextureEditor", "ResolveAlphaChannel", "We don't have alpha information in the source texture data - in the Details tab there is a \"Detect\" button in the quick info that can resolve this."))
+				.Margin(FMargin(2, 0, 2, 0));
 		}
 	}
 
@@ -1527,7 +1740,7 @@ void FTextureEditorToolkit::CreateInternalWidgets()
 
 	TextureViewport = SNew(STextureEditorViewport, SharedThis(this));
 
-	FormatsTabContainer = SNew(SVerticalBox)
+	PlatformsTabContainer = SNew(SVerticalBox)
 	+ SVerticalBox::Slot()
 	.AutoHeight()
 	.Padding(4.0f)
@@ -1537,17 +1750,41 @@ void FTextureEditorToolkit::CreateInternalWidgets()
 		.SelectionMode(ESelectionMode::None)
 		.HeaderRow(
 			SNew(SHeaderRow)
-			+SHeaderRow::Column(FormatsColumnId_Platform)
+			+SHeaderRow::Column(PlatformsColumnId_Platform)
 			.DefaultLabel(LOCTEXT("TextureEditor_PlatformNameCol", "Platform"))
-			.FillWidth(.3f)
-			+SHeaderRow::Column(FormatsColumnId_TextureFormat)
+			.FillWidth(.2f)
+			+SHeaderRow::Column(PlatformsColumnId_TextureFormat)
 			.DefaultLabel(LOCTEXT("TextureEditor_TextureFormatCol", "Texture Format"))
-			.FillWidth(.4f)
-			+SHeaderRow::Column(FormatsColumnId_PixelFormat)
+			.FillWidth(.2f)
+			+SHeaderRow::Column(PlatformsColumnId_PixelFormat)
 			.DefaultLabel(LOCTEXT("TextureEditor_PixelFormatCol", "Pixel Format"))
 			.DefaultTooltip(LOCTEXT("TextureEditor_PixelFormatColTT", "The pixel format that will be chosen for this platform. For texture formats that automatically select alpha or not (e.g. AutoDXT) this might not be known. Click \"Detect\" in the details panel to run the alpha scan and get the actual format"))
-			.FillWidth(.3f)
-		)
+			.FillWidth(.15f)
+			+SHeaderRow::Column(PlatformsColumnId_LODBias)
+			.DefaultLabel(LOCTEXT("TextureEditor_LODBiasCol", "LODBias"))
+			.DefaultTooltip(LOCTEXT("TextureEditor_LODBiasColTT", "The resolved LOD bias for this platform"))
+			.FillWidth(.1f)
+			+SHeaderRow::Column(PlatformsColumnId_Dims)
+			.DefaultLabel(LOCTEXT("TextureEditor_DimsCol", "Size"))
+			.DefaultTooltip(LOCTEXT("TextureEditor_DimsColTT", "The size of the top mip after LOD bias is applied."))
+			.FillWidth(.1f)
+			+SHeaderRow::Column(PlatformsColumnId_NumMips)
+			.DefaultLabel(LOCTEXT("TextureEditor_NumMipsCol", "Mips"))
+			.DefaultTooltip(LOCTEXT("TextureEditor_NumMipsColTT", "The number of mips after LOD bias, with the associated number of mips assigned optional, streaming, and inline. Additionally, if the platform has packed mip tails the number of mips in that mip tail is shown."))
+			.FillWidth(.1f)
+			+SHeaderRow::Column(PlatformsColumnId_SizeOptional)
+			.DefaultLabel(LOCTEXT("TextureEditor_MipBytesSizeCol_Optional", "Optional"))
+			.DefaultTooltip(LOCTEXT("TextureEditor_MipBytesSizeCol_OptionalTT", "The number of bytes of optional mips"))
+			.FillWidth(.15f)
+			+SHeaderRow::Column(PlatformsColumnId_SizeStreaming)
+			.DefaultLabel(LOCTEXT("TextureEditor_MipBytesSizeCol_Streaming", "Streaming"))
+			.DefaultTooltip(LOCTEXT("TextureEditor_MipBytesSizeCol_StreamingTT", "The number of bytes of streaming mips"))
+			.FillWidth(.15f)
+			+SHeaderRow::Column(PlatformsColumnId_SizeInline)
+			.DefaultLabel(LOCTEXT("TextureEditor_MipBytesSizeCol_Inline", "Inline"))
+			.DefaultTooltip(LOCTEXT("TextureEditor_MipBytesSizeCol_InlineTT", "The number of bytes of inline mips. If this is larger than you expect it's likely due to either platform alignment requirements or packed mip tails."))
+			.FillWidth(.15f)
+	)
 		.OnGenerateRow(this, &FTextureEditorToolkit::MakePlatformTextureFormatWidget)
 	]; // end format details.
 
@@ -2855,14 +3092,14 @@ void FTextureEditorToolkit::HandleSettingsActionExecute( )
 }
 
 
-TSharedRef<SDockTab> FTextureEditorToolkit::HandleTabSpawnerSpawnFormats(const FSpawnTabArgs& Args)
+TSharedRef<SDockTab> FTextureEditorToolkit::HandleTabSpawnerSpawnPlatforms(const FSpawnTabArgs& Args)
 {
-	check(Args.GetTabId() == FormatsTabId);
+	check(Args.GetTabId() == PlatformsTabId);
 
 	TSharedRef<SDockTab> SpawnedTab = SNew(SDockTab)
-	.Label(LOCTEXT("TextureFormatsTitle", "Formats"))
+	.Label(LOCTEXT("TexturePlatformsTitle", "Platforms"))
 	[
-		FormatsTabContainer.ToSharedRef()
+		PlatformsTabContainer.ToSharedRef()
 	];
 
 	return SpawnedTab;
diff --git a/Engine/Source/Editor/TextureEditor/Private/TextureEditorToolkit.h b/Engine/Source/Editor/TextureEditor/Private/TextureEditorToolkit.h
--- a/Engine/Source/Editor/TextureEditor/Private/TextureEditorToolkit.h
+++ b/Engine/Source/Editor/TextureEditor/Private/TextureEditorToolkit.h
@@ -334,7 +334,7 @@ private:
 	TSharedRef<SDockTab> HandleTabSpawnerSpawnProperties(const FSpawnTabArgs& Args);
 	TSharedRef<SDockTab> HandleTabSpawnerSpawnViewport( const FSpawnTabArgs& Args );
 	TSharedRef<SDockTab> HandleTabSpawnerSpawnOodle(const FSpawnTabArgs& Args);
-	TSharedRef<SDockTab> HandleTabSpawnerSpawnFormats(const FSpawnTabArgs& Args);
+	TSharedRef<SDockTab> HandleTabSpawnerSpawnPlatforms(const FSpawnTabArgs& Args);
 
 	// Callback for toggling the Texture Border action.
 	void HandleTextureBorderActionExecute( );
@@ -404,7 +404,7 @@ private:
 
 	/** Oodle tab */
 	TSharedPtr<SVerticalBox> OodleTabContainer;
-	TSharedPtr<SVerticalBox> FormatsTabContainer;
+	TSharedPtr<SVerticalBox> PlatformsTabContainer;
 
 	/** Properties tab */
 	TSharedPtr<SVerticalBox> TextureProperties;
@@ -526,10 +526,22 @@ public:
 		// Pixel formats for whether we detected alpha in the texture source. Only
 		// matters for Auto* texture formats. If we don't know the texture source alpha
 		// information we show both.
-		EPixelFormat PixelFormat_AlphaTrue;
-		EPixelFormat PixelFormat_AlphaFalse;
-		bool bKnowAlphaTransparency;
-		bool bHasTransparentAlpha;
+		EPixelFormat PixelFormat_AlphaTrue = PF_Unknown;
+		EPixelFormat PixelFormat_AlphaFalse = PF_Unknown;
+		bool bKnowAlphaTransparency = false;
+		bool bHasTransparentAlpha = false;
+		bool bIsVirtual = false;
+		int32 Width = 0;
+		int32 Height = 0;
+		int32 NumMips_TotalAfterLOD = 0;
+		int32 NumMips_Optional = 0;
+		int32 NumMips_Streaming = 0;
+		int32 NumMips_Tail = 0;
+		int8 LODBias = 0;
+		uint64 MipChainSizeInBytes_Streaming = 0;
+		uint64 MipChainSizeInBytes_Inline = 0;
+		uint64 MipChainSizeInBytes_Optional = 0;
+		TArray<uint64, TInlineAllocator<MAX_TEXTURE_MIP_COUNT>> MipSizesInBytes;
 	};
 private:
 
@@ -606,5 +618,5 @@ private:
 	static const FName ViewportTabId;
 	static const FName PropertiesTabId;
 	static const FName OodleTabId;
-	static const FName FormatsTabId;
+	static const FName PlatformsTabId;
 };
diff --git a/Engine/Source/Runtime/Engine/Private/TextureDerivedData.cpp b/Engine/Source/Runtime/Engine/Private/TextureDerivedData.cpp
--- a/Engine/Source/Runtime/Engine/Private/TextureDerivedData.cpp
+++ b/Engine/Source/Runtime/Engine/Private/TextureDerivedData.cpp
@@ -22,6 +22,7 @@
 #include "VT/VirtualTextureBuiltData.h"
 #include "Misc/ConfigCacheIni.h"
 #include "RenderingThread.h"
+#include "Interfaces/ITextureFormat.h"
 
 #if WITH_EDITOR
 
@@ -36,7 +37,6 @@
 #include "ImageUtils.h"
 #include "Interfaces/ITargetPlatform.h"
 #include "Interfaces/ITargetPlatformManagerModule.h"
-#include "Interfaces/ITextureFormat.h"
 #include "ProfilingDebugging/CookStats.h"
 #include "UObject/ArchiveCookContext.h"
 #include "VT/LightmapVirtualTexture.h"
@@ -2583,33 +2583,8 @@ int32 FTexturePlatformData::GetNumNonStreamingMips(bool bIsStreamingPossible) co
 	}
 	else
 	{
-		int32 MipCount = Mips.Num();
-		int32 NumNonStreamingMips = 1;
 		// MipCount >= 2 and bIsStreamingPossible
-
-		// Take in to account the min resident limit.
-		NumNonStreamingMips = FMath::Max(NumNonStreamingMips, (int32)GetNumMipsInTail());
-		NumNonStreamingMips = FMath::Max(NumNonStreamingMips, UTexture2D::GetStaticMinTextureResidentMipCount());
-		NumNonStreamingMips = FMath::Min(NumNonStreamingMips, MipCount);
-
-		if ( RequiresBlock4Alignment(PixelFormat) )
-		{
-			// ensure the top non-streamed mip (and all streamed mips) size is >= BlockSize (and a multiple of block size!)
-			
-			// note: this is not right for non pow 2; NeverStream should set !bIsStreamingPossible in that case
-			if ( FMath::IsPowerOfTwo(Mips[0].SizeX) && FMath::IsPowerOfTwo(Mips[0].SizeY) )
-			{
-				NumNonStreamingMips = FMath::Max<int32>(NumNonStreamingMips, MipCount - FPlatformMath::FloorLog2(Mips[0].SizeX / 4));
-				NumNonStreamingMips = FMath::Max<int32>(NumNonStreamingMips, MipCount - FPlatformMath::FloorLog2(Mips[0].SizeY / 4));
-			}
-			else
-			{
-				// should have !bIsStreamingPossible, so we should not hit this branch, but it's not reliable
-				NumNonStreamingMips = MipCount;
-			}
-		}
-
-		return NumNonStreamingMips;
+		return GetNumNonStreamingMipsDirect(Mips.Num(), Mips[0].SizeX, Mips[0].SizeY, PixelFormat, GetNumMipsInTail(), UTexture2D::GetStaticMinTextureResidentMipCount());
 	}
 }
 
