Texture : GetTextureStreamingMethod to expose to Blueprint+python
returns trinary current platform streaming method

diff --git a/Engine/Source/Editor/TextureEditor/Private/TextureEditorToolkit.cpp b/Engine/Source/Editor/TextureEditor/Private/TextureEditorToolkit.cpp
--- a/Engine/Source/Editor/TextureEditor/Private/TextureEditorToolkit.cpp
+++ b/Engine/Source/Editor/TextureEditor/Private/TextureEditorToolkit.cpp
@@ -1109,9 +1109,24 @@ void FTextureEditorToolkit::PopulateQuickInfo( )
 		SizeText->SetText(FText::Format(NSLOCTEXT("TextureEditor", "QuickInfo_ResourceSize", "Resource Size: {0} KB"), FText::AsNumber(FMath::DivideAndRoundNearest(ResourceSize, (int64)1024), &FormatOptions)));
 	}
 
-	FText Method = Texture->IsCurrentlyVirtualTextured() ? NSLOCTEXT("TextureEditor", "QuickInfo_MethodVirtualStreamed", "Virtual Streamed")
-													: (!Texture->IsStreamable() ? NSLOCTEXT("TextureEditor", "QuickInfo_MethodNotStreamed", "Not Streamed") 
-																			: NSLOCTEXT("TextureEditor", "QuickInfo_MethodStreamed", "Streamed") );
+	ETextureStreamingMethod StreamingMethod = Texture->GetTextureStreamingMethod();
+	FText Method;
+	switch(StreamingMethod)
+	{
+		default:
+		case TSM_Invalid:
+			Method = NSLOCTEXT("TextureEditor", "QuickInfo_MethodInvalidStreamed", "Invalid");
+			break;
+		case TSM_NotStreamed:
+			Method = NSLOCTEXT("TextureEditor", "QuickInfo_MethodNotStreamed", "Not Streamed");
+			break;
+		case TSM_Streamed:
+			Method = NSLOCTEXT("TextureEditor", "QuickInfo_MethodStreamed", "Streamed");
+			break;
+		case TSM_VirtualStreamed:
+			Method = NSLOCTEXT("TextureEditor", "QuickInfo_MethodVirtualStreamed", "Virtual Streamed");
+			break;
+	}
 
 	MethodText->SetText(FText::Format(NSLOCTEXT("TextureEditor", "QuickInfo_Method", "Method: {0}"), Method));
 	LODBiasText->SetText(FText::Format(NSLOCTEXT("TextureEditor", "QuickInfo_LODBias", "Combined LOD Bias: {0}"), FText::AsNumber(Texture->GetCachedLODBias())));
diff --git a/Engine/Source/Runtime/Engine/Classes/Engine/Texture.h b/Engine/Source/Runtime/Engine/Classes/Engine/Texture.h
--- a/Engine/Source/Runtime/Engine/Classes/Engine/Texture.h
+++ b/Engine/Source/Runtime/Engine/Classes/Engine/Texture.h
@@ -1750,6 +1750,10 @@ public:
 		return false;
 	}
 
+	/** Returns the type of streaming the texture is currently using on the current platform. (None/Streamed/Virtual) **/	
+	UFUNCTION(BlueprintCallable, Category = Texture)
+	ENGINE_API ETextureStreamingMethod GetTextureStreamingMethod() const;
+
 	/** Returns the virtual texture build settings. */
 	ENGINE_API virtual void GetVirtualTextureBuildSettings(struct FVirtualTextureBuildSettings& OutSettings) const;
 
diff --git a/Engine/Source/Runtime/Engine/Classes/Engine/TextureDefines.h b/Engine/Source/Runtime/Engine/Classes/Engine/TextureDefines.h
--- a/Engine/Source/Runtime/Engine/Classes/Engine/TextureDefines.h
+++ b/Engine/Source/Runtime/Engine/Classes/Engine/TextureDefines.h
@@ -506,6 +506,15 @@ enum ETextureCompressionQuality : int
 	TCQ_MAX,
 };
 
+UENUM()
+enum ETextureStreamingMethod : int
+{
+	TSM_Invalid = 0,   // Unknown / not set
+	TSM_NotStreamed,   // not streamed (fully resident all the time)
+	TSM_Streamed,      // traditional whole-mip streaming
+	TSM_VirtualStreamed  // virtual texture (VT)
+};
+
 
 namespace UE
 {
diff --git a/Engine/Source/Runtime/Engine/Private/Texture.cpp b/Engine/Source/Runtime/Engine/Private/Texture.cpp
--- a/Engine/Source/Runtime/Engine/Private/Texture.cpp
+++ b/Engine/Source/Runtime/Engine/Private/Texture.cpp
@@ -1932,6 +1932,30 @@ bool UTexture::IsCandidateForTextureStreamingOnPlatformDuringCook(const ITargetP
 }
 #endif
 
+
+ETextureStreamingMethod UTexture::GetTextureStreamingMethod() const
+{
+	// returns the streaming method currently in use on this platform
+	//	is not a general property of the texture
+	//	can vary during texture lifetime (eg. PostInit streaming setup)
+
+	if ( IsCurrentlyVirtualTextured() )
+	{
+		return TSM_VirtualStreamed;
+	}
+	else if ( IsStreamable() )
+	{
+		// StreamableRenderAsset.h
+		// this will only be true after GetResourcePostInitState sets up lods
+		//	see also IsPossibleToStream
+		return TSM_Streamed;
+	}
+	else
+	{
+		return TSM_NotStreamed;
+	}
+}
+
 FStreamableRenderResourceState UTexture::GetResourcePostInitState(const FTexturePlatformData* PlatformData, bool bAllowStreaming, int32 MinRequestMipCount, int32 MaxMipCount, bool bSkipCanBeLoaded) const
 {
 	// Async caching of PlatformData must be done before calling this
