Show warning dialog when trying to convert non power of 2 texture to virtual texture in the texture editor.
#rnx
#Synced-CL 49674283

diff --git a/Engine/Source/Editor/TextureEditor/Private/Customizations/TextureDetailsCustomization.cpp b/Engine/Source/Editor/TextureEditor/Private/Customizations/TextureDetailsCustomization.cpp
--- a/Engine/Source/Editor/TextureEditor/Private/Customizations/TextureDetailsCustomization.cpp
+++ b/Engine/Source/Editor/TextureEditor/Private/Customizations/TextureDetailsCustomization.cpp
@@ -180,11 +180,21 @@ void FTextureDetails::CustomizeDetails(IDetailLayoutBuilder& DetailBuilder)
 				{
 					if (UTexture2D* Texture = Cast<UTexture2D>(WeakObject.Get()))
 					{
-						TArray<UTexture2D*> Textures;
-						Textures.Add(Texture);
-						IAssetTools& AssetTools = FModuleManager::LoadModuleChecked<FAssetToolsModule>("AssetTools").Get();
-						AssetTools.ConvertVirtualTextures(Textures, Texture->VirtualTextureStreaming);
-						return FReply::Handled();
+						if (!Texture->Source.AreAllBlocksPowerOfTwo() && Texture->PowerOfTwoMode == ETexturePowerOfTwoSetting::None)
+						{
+							FMessageDialog::Open(
+								EAppMsgType::Ok,
+								LOCTEXT("Dialog_ConvertTexturePow2_Msg", "Failed to convert a non Power of Two sized texture.\nSet a valid Padding and Resizing mode."),
+								LOCTEXT("Dialog_ConvertTexture_Title", "Convert Texture"));
+						}
+						else
+						{
+							TArray<UTexture2D*> Textures;
+							Textures.Add(Texture);
+							IAssetTools& AssetTools = FModuleManager::LoadModuleChecked<FAssetToolsModule>("AssetTools").Get();
+							AssetTools.ConvertVirtualTextures(Textures, Texture->VirtualTextureStreaming);
+							return FReply::Handled();
+						}
 					}
 					return FReply::Unhandled();
 				})
