TextureEditorToolkit :  PopulateQuickInfo  NOP debugging

#Synced-CL 50735006

diff --git a/Engine/Source/Editor/TextureEditor/Private/TextureEditorToolkit.cpp b/Engine/Source/Editor/TextureEditor/Private/TextureEditorToolkit.cpp
--- a/Engine/Source/Editor/TextureEditor/Private/TextureEditorToolkit.cpp
+++ b/Engine/Source/Editor/TextureEditor/Private/TextureEditorToolkit.cpp
@@ -610,6 +610,8 @@ double FTextureEditorToolkit::GetCustomZoomLevel( ) const
 
 void FTextureEditorToolkit::PopulateQuickInfo( )
 {
+	check( Texture != nullptr );
+
 	if (Texture->IsDefaultTexture())
 	{
 		ImportedText->SetText(NSLOCTEXT("TextureEditor", "QuickInfo_Imported_NA", "Imported: Computing..."));
@@ -1172,11 +1174,17 @@ void FTextureEditorToolkit::PopulateQuickInfo( )
 		for (TObjectIterator<USceneCaptureComponent2D> It; It; ++It)
 		{
 			USceneCaptureComponent2D* SceneCaptureComponent = *It;
-			if (SceneCaptureComponent->TextureTarget == Texture && SceneCaptureComponent->CaptureMemorySize.IsValid())
+			check( SceneCaptureComponent != nullptr );
+			
+			if (SceneCaptureComponent->TextureTarget == Texture )
 			{
-				// Could have multiple scene captures blending to the same render target -- add them all up, but only show the name of the last one...
-				CaptureSize += SceneCaptureComponent->CaptureMemorySize->Size;
-				CaptureName = SceneCaptureComponent->GetOwner()->GetFName();
+				TRefCountPtr<FSceneCaptureMemorySize> CaptureMemorySize = SceneCaptureComponent->CaptureMemorySize;
+				if ( CaptureMemorySize.IsValid())
+				{
+					// Could have multiple scene captures blending to the same render target -- add them all up, but only show the name of the last one...
+					CaptureSize += CaptureMemorySize->Size;
+					CaptureName = SceneCaptureComponent->GetOwner() ? SceneCaptureComponent->GetOwner()->GetFName() : FName();
+				}
 			}
 		}
 	}
